[ModelContextProtocol] Move `ModelContextProtocol` plugin to `Engine/Plugins/Experimental`.

#Synced-CL 52053324

diff --git a/Engine/Plugins/Experimental/ModelContextProtocol/Source/ModelContextProtocol/Private/ModelContextProtocolEngineSubsystem.cpp b/Engine/Plugins/Experimental/ModelContextProtocol/Source/ModelContextProtocol/Private/ModelContextProtocolEngineSubsystem.cpp
new file mode 100644
--- /dev/null
+++ b/Engine/Plugins/Experimental/ModelContextProtocol/Source/ModelContextProtocol/Private/ModelContextProtocolEngineSubsystem.cpp
@@ -0,0 +1,1216 @@
+﻿// Copyright Epic Games, Inc. All Rights Reserved.
+
+#include "ModelContextProtocolEngineSubsystem.h"
+#include "IModelContextProtocolModule.h"
+#include "IModelContextProtocolResourceProvider.h"
+#include "ModelContextProtocol.h"
+#include "ModelContextProtocolSettings.h"
+#include "ModelContextProtocolToolAsyncAction.h"
+#include "ModelContextProtocolToolLibrary.h"
+
+#include "Dom/JsonObject.h"
+#include "Engine/Engine.h"
+#include "Misc/Base64.h"
+#include "HttpServerModule.h"
+#include "IHttpRouter.h"
+#include "JsonDomBuilder.h"
+#include "Misc/CoreDelegates.h"
+#include "Policies/CondensedJsonPrintPolicy.h"
+#include "Serialization/JsonReader.h"
+#include "Serialization/JsonSerializer.h"
+#include "UObject/UObjectIterator.h"
+
+#if WITH_EDITOR
+#include "Editor.h"
+#include "ToolsetRegistry/ToolsetRegistry.h"
+#include "ToolsetRegistry/ToolsetRegistrySubsystem.h"
+#endif
+
+#include UE_INLINE_GENERATED_CPP_BY_NAME(ModelContextProtocolEngineSubsystem)
+
+namespace UE::ModelContextProtocol::Private
+{
+	static constexpr const TCHAR* McpSessionIdHeader = TEXT("Mcp-Session-Id");
+	static constexpr const TCHAR* McpProtocolVersionHeader = TEXT("Mcp-Protocol-Version");
+
+	bool GetProtocolVersionFromRequestHeaders(const FHttpServerRequest& Request, FString& OutProtocolVersion)
+	{
+		if (const TArray<FString>* Headers = Request.Headers.Find(McpProtocolVersionHeader))
+		{
+			if (!Headers->IsEmpty())
+			{
+				OutProtocolVersion = (*Headers)[0];
+				return true;
+			}
+		}
+		return false;
+	}
+
+	void CompleteWithResponseCode(const FHttpResultCallback& OnComplete, EHttpServerResponseCodes ResponseCode);
+
+	/**
+	 * Validates the Origin header to prevent DNS rebinding attacks per MCP spec.
+	 * Allows requests with no Origin (non-browser clients), or Origin matching localhost variants.
+	 * Returns true if the request is allowed, false if it was rejected (and OnComplete was called).
+	 */
+	bool ValidateOriginHeader(const FHttpServerRequest& Request, const FHttpResultCallback& OnComplete)
+	{
+		const TArray<FString>* OriginHeaders = Request.Headers.Find(TEXT("Origin"));
+		if (!OriginHeaders || OriginHeaders->IsEmpty())
+		{
+			// No Origin header — non-browser client, allow
+			return true;
+		}
+
+		const FString& Origin = (*OriginHeaders)[0];
+
+		// Extract the host portion from the Origin (scheme://host[:port])
+		FString OriginHost;
+		FString SchemeAndRest;
+		if (Origin.Split(TEXT("://"), &SchemeAndRest, &OriginHost))
+		{
+			// Strip port if present
+			FString Host;
+			FString Port;
+			if (!OriginHost.Split(TEXT(":"), &Host, &Port))
+			{
+				Host = OriginHost;
+			}
+
+			if (Host == TEXT("localhost") || Host == TEXT("127.0.0.1") || Host == TEXT("[::1]"))
+			{
+				return true;
+			}
+		}
+
+		UE_LOG(LogModelContextProtocol, Warning, TEXT("Rejected request with disallowed Origin: %s"), *Origin);
+		CompleteWithResponseCode(OnComplete, EHttpServerResponseCodes::Forbidden);
+		return false;
+	}
+
+	FAutoConsoleCommand CommandStartServer = FAutoConsoleCommand(
+		TEXT("ModelContextProtocol.StartServer"),
+		TEXT("Explicitly starts the MCP HTTP server"),
+		FConsoleCommandWithArgsDelegate::CreateLambda([](const TArray<FString>& Args)
+		{
+			if (GEngine && GEngine->IsInitialized())
+			{
+				if (UModelContextProtocolEngineSubsystem* Subsystem = GEngine->GetEngineSubsystem<UModelContextProtocolEngineSubsystem>(); ensure(Subsystem))
+				{
+					Subsystem->StartServer();
+				}
+			}
+		}),
+		ECVF_Cheat);
+
+	FAutoConsoleCommand CommandStopServer = FAutoConsoleCommand(
+		TEXT("ModelContextProtocol.StopServer"),
+		TEXT("Explicitly stops the MCP HTTP server"),
+		FConsoleCommandWithArgsDelegate::CreateLambda([](const TArray<FString>& Args)
+		{
+			if (GEngine && GEngine->IsInitialized())
+			{
+				if (UModelContextProtocolEngineSubsystem* Subsystem = GEngine->GetEngineSubsystem<UModelContextProtocolEngineSubsystem>(); ensure(Subsystem))
+				{
+					Subsystem->StopServer();
+				}
+			}
+		}),
+		ECVF_Cheat);
+	
+	enum class EJsonRpcErrorCode : int32
+	{
+		/**
+		 * Invalid JSON was received by the server.
+		 * An error occurred on the server while parsing the JSON text.
+		 */
+		ParseError = -32700,
+		
+		/** The JSON sent is not a valid Request object. */
+		InvalidRequest = -32600,
+
+		/** The method does not exist / is not available. */
+		MethodNotFound = -32601,
+
+		/** Invalid method parameter(s). */
+		InvalidParams = -32602,
+
+		/** Internal JSON-RPC error. */
+		InternalError = -32603,
+
+		/** Resource not found (MCP-specific). */
+		ResourceNotFound = -32002
+	};
+	
+	TSharedPtr<FJsonObject> GetJsonObjectFromRequestBody(const TArray<uint8>& InRequestBody)
+	{
+		const FUtf8StringView IncomingRequestBody(reinterpret_cast<const UTF8CHAR*>(InRequestBody.GetData()), InRequestBody.Num());
+		TSharedPtr<FJsonObject> BodyObject = MakeShared<FJsonObject>();
+		TSharedRef<TJsonReader<UTF8CHAR>> JsonReader = TJsonReaderFactory<UTF8CHAR>::CreateFromView(IncomingRequestBody);
+
+		if (FJsonSerializer::Deserialize(JsonReader, BodyObject) && BodyObject.IsValid())
+		{
+			return BodyObject;
+		}
+
+		return nullptr;
+	}
+
+	TUniquePtr<FHttpServerResponse> CreateJsonRpcErrorResponse(const TSharedPtr<FJsonValue>& RequestId, int32 ErrorCode, const FString& ErrorMessage, EHttpServerResponseCodes ResponseCode = EHttpServerResponseCodes::BadRequest)
+	{
+		FJsonDomBuilder::FObject ErrorObject;
+		ErrorObject.Set(TEXT("code"), ErrorCode);
+		ErrorObject.Set(TEXT("message"), ErrorMessage);
+		
+		FJsonDomBuilder::FObject ResponseObject;
+		ResponseObject.Set(TEXT("jsonrpc"), UE::ModelContextProtocol::JsonRpcVersion);
+		ResponseObject.Set(TEXT("id"), RequestId.IsValid() ? RequestId : MakeShared<FJsonValueNull>());
+		ResponseObject.Set(TEXT("error"), ErrorObject);
+		
+		FUtf8String ResponseStr;
+		TSharedRef<TJsonWriter<UTF8CHAR>> JsonWriter = TJsonWriterFactory<UTF8CHAR>::Create(&ResponseStr);
+		FJsonSerializer::Serialize(ResponseObject.AsJsonObject(), JsonWriter);
+		
+		TUniquePtr<FHttpServerResponse> Response = FHttpServerResponse::Create(MoveTemp(ResponseStr), TEXT("application/json"));
+		Response->Code = ResponseCode;
+		return Response;
+	}
+
+	
+	TUniquePtr<FHttpServerResponse> CreateJsonRpcResultResponse(const TSharedPtr<FJsonValue>& RequestId, const TSharedPtr<FJsonValue>& Result, const FString* SessionId = nullptr)
+	{
+		if (!ensureMsgf(RequestId.IsValid(), TEXT("JSON-RPC responses *MUST* have a valid request id. Requests without an id must be treated as notifications and should not be replied to.")))
+		{
+			return CreateJsonRpcErrorResponse(RequestId, static_cast<int32>(EJsonRpcErrorCode::InternalError), TEXT("Replying to request with unknown RequestId"));
+		}
+		
+		FJsonDomBuilder::FObject ResponseObject;
+		ResponseObject.Set(TEXT("jsonrpc"), JsonRpcVersion);
+		ResponseObject.Set(TEXT("id"), RequestId);
+		ResponseObject.Set(TEXT("result"), Result);
+		
+		FUtf8String ResponseStr;
+		TSharedRef<TJsonWriter<UTF8CHAR>> JsonWriter = TJsonWriterFactory<UTF8CHAR>::Create(&ResponseStr);
+		FJsonSerializer::Serialize(ResponseObject.AsJsonObject(), JsonWriter);
+		
+		TUniquePtr<FHttpServerResponse> ServerResponse = FHttpServerResponse::Create(MoveTemp(ResponseStr), TEXT("application/json"));
+		if (SessionId)
+		{
+			ServerResponse->Headers.Add(McpSessionIdHeader, {*SessionId});
+		}
+		return ServerResponse;
+	}
+
+	void CompleteWithError(const FHttpResultCallback& OnComplete, const TSharedPtr<FJsonValue>& RequestId, UE::ModelContextProtocol::Private::EJsonRpcErrorCode ErrorCode, const FString& ErrorMessage, EHttpServerResponseCodes ResponseCode = EHttpServerResponseCodes::BadRequest)
+	{
+		UE_LOG(LogModelContextProtocol, Error, TEXT("%s"), *ErrorMessage);
+		OnComplete(CreateJsonRpcErrorResponse(RequestId, static_cast<int32>(ErrorCode), ErrorMessage, ResponseCode));
+	}
+
+	void CompleteWithResult(const FHttpResultCallback& OnComplete, const TSharedPtr<FJsonValue>& RequestId, const TSharedPtr<FJsonValue>& Result, const FString* SessionId = nullptr)
+	{
+		OnComplete(CreateJsonRpcResultResponse(RequestId, Result, SessionId));
+	}
+	
+	void CompleteWithResult(const FHttpResultCallback& OnComplete, const TSharedPtr<FJsonValue>& RequestId, const FJsonDomBuilder::FObject& Result, const FString* SessionId = nullptr)
+	{
+		OnComplete(CreateJsonRpcResultResponse(RequestId, Result.AsJsonValue(), SessionId));
+	}
+
+	template <typename ResultType>
+	void CompleteWithResult(const FHttpResultCallback& OnComplete, const TSharedPtr<FJsonValue>& RequestId, const ResultType& Result, const FString* SessionId = nullptr)
+	{
+		TSharedPtr<FJsonObject> JsonResult = FJsonObjectConverter::UStructToJsonObject<ResultType>(Result);
+		if (!ensure(JsonResult.IsValid()))
+		{
+			CompleteWithError(
+				OnComplete, RequestId, EJsonRpcErrorCode::InternalError, FString::Printf(TEXT("Error forming response")));
+			return;
+		}
+
+		TSharedPtr<FJsonValue> JsonResultValue = MakeShared<FJsonValueObject>(JsonResult);
+		CompleteWithResult(OnComplete, RequestId, JsonResultValue, SessionId);
+	}
+
+	void CompleteWithResponseCode(const FHttpResultCallback& OnComplete, EHttpServerResponseCodes ResponseCode)
+	{
+		TUniquePtr<FHttpServerResponse> Response = MakeUnique<FHttpServerResponse>();
+		Response->Code = ResponseCode;
+		OnComplete(MoveTemp(Response));
+	}
+
+	template <typename ParamsType>
+	bool ParseParamsOrError(const TCHAR* MethodName, const TSharedPtr<FJsonValue>& RequestId, const FHttpResultCallback& OnComplete, const TSharedPtr<FJsonObject>& JsonParams, ParamsType& OutParams)
+	{
+		// Parse params
+		if (!JsonParams.IsValid())
+		{
+			CompleteWithError(
+				OnComplete, RequestId, EJsonRpcErrorCode::InvalidParams, FString::Printf(TEXT("Expected params for call to \"%s\""), MethodName));
+			return false;
+		}
+		FText OutFailReason;
+		if (!FJsonObjectConverter::JsonObjectToUStruct<ParamsType>(JsonParams.ToSharedRef(), &OutParams, /*CheckFlags*/CPF_None, /*SkipFlags*/CPF_None, /*bStrictMode*/false, &OutFailReason))
+		{
+			CompleteWithError(
+				OnComplete, RequestId, EJsonRpcErrorCode::InvalidParams, FString::Printf(TEXT("Invalid %s params: %s"), MethodName, *OutFailReason.ToString()));
+			return false;
+		}
+
+		return true;
+	}
+
+	bool GetSessionIdFromRequestHeaders(const FHttpServerRequest& Request, FString& OutSessionId)
+	{
+		if (const TArray<FString>* McpSessionIdHeaders = Request.Headers.Find(McpSessionIdHeader))
+		{
+			if (!McpSessionIdHeaders->IsEmpty())
+			{
+				OutSessionId = (*McpSessionIdHeaders)[0];
+				return true;
+			}
+		}
+		return false;
+	}
+
+	FUtf8String FormatSSEMessage(const FUtf8String& Message)
+	{
+		return UTF8TEXT("event: message\r\ndata: ") + Message + UTF8TEXT("\r\n\r\n");
+	}
+
+	TUniquePtr<FHttpServerResponse> CreateJsonRpcProgressResponse(const TSharedRef<FJsonValue>& ProgressToken, int32 Progress)
+	{
+		FJsonDomBuilder::FObject ParamsObject;
+		ParamsObject.Set(TEXT("progressToken"), ProgressToken);
+		ParamsObject.Set(TEXT("progress"), Progress);
+
+		FJsonDomBuilder::FObject ResponseObject;
+		ResponseObject.Set(TEXT("jsonrpc"), UE::ModelContextProtocol::JsonRpcVersion);
+		ResponseObject.Set(TEXT("method"), TEXT("notifications/progress"));
+		ResponseObject.Set(TEXT("params"), ParamsObject);
+
+		FUtf8String ResponseStr;
+		TSharedRef<TJsonWriter<UTF8CHAR, TCondensedJsonPrintPolicy<UTF8CHAR>>> JsonWriter = TJsonWriterFactory<UTF8CHAR, TCondensedJsonPrintPolicy<UTF8CHAR>>::Create(&ResponseStr);
+		FJsonSerializer::Serialize(ResponseObject.AsJsonObject(), JsonWriter);
+
+		TUniquePtr<FHttpServerResponse> Response = FHttpServerResponse::Create(UE::ModelContextProtocol::Private::FormatSSEMessage(ResponseStr), UE::ModelContextProtocol::ContentTypeEventStream);
+		Response->Code = EHttpServerResponseCodes::Ok;
+		EnumAddFlags(Response->Flags, EHttpServerResponseFlags::MultipleWriteStream | EHttpServerResponseFlags::HasAdditionalWrites | EHttpServerResponseFlags::SkipHeaderWrite);
+
+		return Response;
+	}
+
+	void SendProgressUpdate(const FHttpResultCallback& ResponseCallback, const TSharedRef<FJsonValue>& ProgressToken, int32 Progress)
+	{
+		ResponseCallback(CreateJsonRpcProgressResponse(ProgressToken, Progress));
+	}
+
+	TUniquePtr<FHttpServerResponse> CreateJsonRpcNotificationResponse(const FString& Method)
+	{
+		FJsonDomBuilder::FObject NotificationObject;
+		NotificationObject.Set(TEXT("jsonrpc"), UE::ModelContextProtocol::JsonRpcVersion);
+		NotificationObject.Set(TEXT("method"), Method);
+		NotificationObject.Set(TEXT("params"), FJsonDomBuilder::FObject());
+
+		FUtf8String ResponseStr;
+		TSharedRef<TJsonWriter<UTF8CHAR, TCondensedJsonPrintPolicy<UTF8CHAR>>> JsonWriter = TJsonWriterFactory<UTF8CHAR, TCondensedJsonPrintPolicy<UTF8CHAR>>::Create(&ResponseStr);
+		FJsonSerializer::Serialize(NotificationObject.AsJsonObject(), JsonWriter);
+
+		TUniquePtr<FHttpServerResponse> Response = FHttpServerResponse::Create(UE::ModelContextProtocol::Private::FormatSSEMessage(ResponseStr), UE::ModelContextProtocol::ContentTypeEventStream);
+		Response->Code = EHttpServerResponseCodes::Ok;
+		EnumAddFlags(Response->Flags, EHttpServerResponseFlags::MultipleWriteStream | EHttpServerResponseFlags::HasAdditionalWrites | EHttpServerResponseFlags::SkipHeaderWrite);
+
+		return Response;
+	}
+
+	struct FPaginationResult
+	{
+		TArray<TSharedPtr<FJsonValue>> Items;
+		TOptional<FString> NextCursor;
+	};
+
+	bool ApplyPagination(const TArray<TSharedPtr<FJsonValue>>& AllItems, const TSharedPtr<FJsonObject>& Params, FPaginationResult& OutResult, FString& OutError)
+	{
+		const int32 PageSize = UE::ModelContextProtocol::PaginationPageSize;
+
+		if (PageSize <= 0)
+		{
+			OutResult.Items = AllItems;
+			return true;
+		}
+
+		int32 Offset = 0;
+
+		if (Params.IsValid())
+		{
+			FString Cursor;
+			if (Params->TryGetStringField(TEXT("cursor"), Cursor) && !Cursor.IsEmpty())
+			{
+				FString DecodedCursor;
+				auto IsAllDigits = [](const FString& String)
+				{
+					for (TCHAR Character : String)
+					{
+						if (!FChar::IsDigit(Character))
+						{
+							return false;
+						}
+					}
+					return true;
+				};
+				if (!FBase64::Decode(Cursor, DecodedCursor) || DecodedCursor.IsEmpty() || !IsAllDigits(DecodedCursor))
+				{
+					OutError = TEXT("Invalid pagination cursor");
+					return false;
+				}
+
+				const int64 DecodedOffset = FCString::Atoi64(*DecodedCursor);
+				if (DecodedOffset < 0 || DecodedOffset > MAX_int32)
+				{
+					OutError = TEXT("Invalid pagination cursor");
+					return false;
+				}
+				Offset = static_cast<int32>(DecodedOffset);
+			}
+		}
+
+		if (Offset > AllItems.Num())
+		{
+			Offset = AllItems.Num();
+		}
+
+		const int32 End = static_cast<int32>(FMath::Min(static_cast<int64>(Offset) + PageSize, static_cast<int64>(AllItems.Num())));
+
+		OutResult.Items.Reserve(End - Offset);
+		for (int32 Index = Offset; Index < End; ++Index)
+		{
+			OutResult.Items.Add(AllItems[Index]);
+		}
+
+		if (End < AllItems.Num())
+		{
+			OutResult.NextCursor = FBase64::Encode(FString::FromInt(End));
+		}
+
+		return true;
+	}
+}
+
+TSharedPtr<FModelContextProtocolSession> UModelContextProtocolEngineSubsystem::FindSession(const FString& SessionId) const
+{
+	const TSharedPtr<FModelContextProtocolSession>* SessionPtr = Sessions.FindByPredicate([&SessionId](const TSharedPtr<FModelContextProtocolSession>& Session)
+		{
+			return Session->ID == SessionId;
+		});
+	return SessionPtr ? *SessionPtr : nullptr;
+}
+
+void UModelContextProtocolEngineSubsystem::StartServer()
+{
+	if (HttpRouter.IsValid())
+	{
+		UE_LOG(LogModelContextProtocol, Warning, TEXT("MCP server is already running on port %u, stopping first."), ActiveServerPort);
+		StopServer();
+	}
+	
+	ActiveServerPort = UE::ModelContextProtocol::GetServerPortNumber();
+	HttpRouter = FHttpServerModule::Get().GetHttpRouter(ActiveServerPort);
+	if (ensure(HttpRouter.IsValid()))
+	{
+		const FHttpPath MCPServerPath = UE::ModelContextProtocol::GetServerUrlPath();
+
+		MainMcpRoute = HttpRouter->BindRoute(MCPServerPath, EHttpServerRequestVerbs::VERB_POST,
+			FHttpRequestHandler::CreateUObject(this, &UModelContextProtocolEngineSubsystem::ProcessPostRequest));
+
+		SseMcpRoute = HttpRouter->BindRoute(MCPServerPath, EHttpServerRequestVerbs::VERB_GET,
+			FHttpRequestHandler::CreateUObject(this, &UModelContextProtocolEngineSubsystem::ProcessGetRequest));
+
+		DeleteMcpRoute = HttpRouter->BindRoute(MCPServerPath, EHttpServerRequestVerbs::VERB_DELETE,
+			FHttpRequestHandler::CreateUObject(this, &UModelContextProtocolEngineSubsystem::ProcessDeleteRequest));
+	}
+		
+	FHttpServerModule::Get().StartAllListeners();
+}
+
+void UModelContextProtocolEngineSubsystem::Initialize(FSubsystemCollectionBase& Collection)
+{
+	Super::Initialize(Collection);
+
+	if (IModelContextProtocolModule* Module = IModelContextProtocolModule::Get())
+	{
+		OnRefreshToolsDelegateHandle = Module->OnRefreshTools().AddUObject(this, &UModelContextProtocolEngineSubsystem::BroadcastToolsListChanged);
+	}
+
+	if (UE::ModelContextProtocol::ShouldAutoStartServer())
+	{
+		StartServer();
+	}
+}
+
+void UModelContextProtocolEngineSubsystem::Deinitialize()
+{
+	if (IModelContextProtocolModule* Module = IModelContextProtocolModule::Get())
+	{
+		Module->OnRefreshTools().Remove(OnRefreshToolsDelegateHandle);
+	}
+
+	StopServer();
+
+	Super::Deinitialize();
+}
+
+void UModelContextProtocolEngineSubsystem::StopServer()
+{
+	if (HttpRouter.IsValid())
+	{
+		if (MainMcpRoute.IsValid())
+		{
+			HttpRouter->UnbindRoute(MainMcpRoute);
+		}
+
+		if (SseMcpRoute.IsValid())
+		{
+			HttpRouter->UnbindRoute(SseMcpRoute);
+		}
+
+		if (DeleteMcpRoute.IsValid())
+		{
+			HttpRouter->UnbindRoute(DeleteMcpRoute);
+		}
+	}
+
+	MainMcpRoute.Reset();
+	SseMcpRoute.Reset();
+	DeleteMcpRoute.Reset();
+	HttpRouter.Reset();
+	ActiveServerPort = 0;
+}
+
+bool UModelContextProtocolEngineSubsystem::IsServerRunning() const
+{
+	return HttpRouter.IsValid();
+}
+
+uint32 UModelContextProtocolEngineSubsystem::GetServerPort() const
+{
+	return ActiveServerPort;
+}
+
+void UModelContextProtocolEngineSubsystem::SetServerPort(uint32 Port)
+{
+	UModelContextProtocolSettings* Settings = GetMutableDefault<UModelContextProtocolSettings>();
+	Settings->ServerPortNumber = Port;
+	Settings->SaveConfig();
+}
+
+void UModelContextProtocolEngineSubsystem::BroadcastToolsListChanged()
+{
+	for (const TSharedPtr<FModelContextProtocolSession>& Session : Sessions)
+	{
+		if (!Session.IsValid() || Session->Status != EModelContextProtocolSessionStatus::Initialized)
+		{
+			continue;
+		}
+
+		for (auto& [RequestId, Context] : Session->ActiveRequests)
+		{
+			if (Context.EventStreamWrite.IsSet())
+			{
+				Context.EventStreamWrite(UE::ModelContextProtocol::Private::CreateJsonRpcNotificationResponse(TEXT("notifications/tools/list_changed")));
+				// Only need to send once per session — all active SSE streams share the same client.
+				break;
+			}
+		}
+	}
+
+	UE_LOG(LogModelContextProtocol, Log, TEXT("Broadcast notifications/tools/list_changed to active sessions"));
+}
+
+bool UModelContextProtocolEngineSubsystem::ProcessPostRequest(const FHttpServerRequest& Request, const FHttpResultCallback& OnComplete)
+{
+	if (!UE::ModelContextProtocol::Private::ValidateOriginHeader(Request, OnComplete))
+	{
+		return true;
+	}
+
+	TSharedPtr<FJsonObject> RequestBodyObject = UE::ModelContextProtocol::Private::GetJsonObjectFromRequestBody(Request.Body);
+	if (!RequestBodyObject.IsValid())
+	{
+		UE::ModelContextProtocol::Private::CompleteWithError(
+			OnComplete, /*RequestId*/{}, UE::ModelContextProtocol::Private::EJsonRpcErrorCode::ParseError, TEXT("Invalid JSON body!"));
+		return true;
+	}
+
+	// jsonrpc
+	FString RequestJsonRpcVersion;
+	if (!RequestBodyObject->TryGetStringField(TEXT("jsonrpc"), RequestJsonRpcVersion) || RequestJsonRpcVersion != UE::ModelContextProtocol::JsonRpcVersion)
+	{
+		UE::ModelContextProtocol::Private::CompleteWithError(
+			OnComplete, /*RequestId*/{}, UE::ModelContextProtocol::Private::EJsonRpcErrorCode::InvalidRequest, 
+			FString::Printf(TEXT("Unexpected jsonrpc version %s. Must be %s"), *RequestJsonRpcVersion, UE::ModelContextProtocol::JsonRpcVersion));
+		return true;
+	}
+
+	// id
+	TSharedPtr<FJsonValue> RequestId = RequestBodyObject->TryGetField(TEXT("id"));
+
+	// method
+	FString Method;
+	if (!RequestBodyObject->TryGetStringField(TEXT("method"), Method))
+	{
+		UE::ModelContextProtocol::Private::CompleteWithError(
+			OnComplete, /*RequestId*/{}, UE::ModelContextProtocol::Private::EJsonRpcErrorCode::InvalidRequest, TEXT("Missing required field \"method\""));
+		return true;
+	}
+
+	// params
+	const TSharedPtr<FJsonObject>* ParamsObject = nullptr;
+	RequestBodyObject->TryGetObjectField(TEXT("params"), ParamsObject);
+
+	return ProcessJsonRpcCall(Request, RequestId, Method, ParamsObject ? *ParamsObject : nullptr, OnComplete);
+}
+
+bool UModelContextProtocolEngineSubsystem::ProcessJsonRpcCall(const FHttpServerRequest& Request, const TSharedPtr<FJsonValue>& RequestId, const FString& Method, const TSharedPtr<FJsonObject>& Params, const FHttpResultCallback& OnComplete)
+{
+	TSharedPtr<FJsonValue> Result;
+
+	// ping
+	if (Method == TEXT("ping"))
+	{
+		return ProcessPingJsonRpcCall(Request, RequestId, OnComplete);
+	}
+	// initialize
+	if (Method == TEXT("initialize"))
+	{
+		return ProcessInitializeJsonRpcCall(Request, RequestId, Params, OnComplete);
+	}
+
+	// Per MCP spec (2025-11-25): validate Mcp-Protocol-Version header on all post-initialize requests.
+	// If present, verify it matches the session's negotiated version.
+	FString HeaderProtocolVersion;
+	if (UE::ModelContextProtocol::Private::GetProtocolVersionFromRequestHeaders(Request, HeaderProtocolVersion))
+	{
+		FString SessionId;
+		if (UE::ModelContextProtocol::Private::GetSessionIdFromRequestHeaders(Request, SessionId))
+		{
+			if (TSharedPtr<FModelContextProtocolSession> Session = FindSession(SessionId))
+			{
+				if (HeaderProtocolVersion != Session->NegotiatedProtocolVersion)
+				{
+					UE::ModelContextProtocol::Private::CompleteWithError(
+						OnComplete, RequestId, UE::ModelContextProtocol::Private::EJsonRpcErrorCode::InvalidRequest,
+						FString::Printf(TEXT("Mcp-Protocol-Version header '%s' does not match negotiated version '%s'"), *HeaderProtocolVersion, *Session->NegotiatedProtocolVersion));
+					return true;
+				}
+			}
+		}
+	}
+
+	// notifications/initialized
+	if (Method == TEXT("notifications/initialized"))
+	{
+		return ProcessInitializedNotificationJsonRpcCall(Request, RequestId, Params, OnComplete);
+	}
+	// notifications/cancelled
+	else if (Method == TEXT("notifications/cancelled"))
+	{
+		return ProcessNotificationCancelledJsonRpcCall(Request, RequestId, Params, OnComplete);
+	}
+	// tools/list
+	else if (Method == TEXT("tools/list"))
+	{
+		return ProcessListToolsJsonRpcCall(Request, RequestId, Params, OnComplete);
+	}
+	// tools/call
+	else if (Method == TEXT("tools/call"))
+	{
+		return ProcessToolCallJsonRpcCall(Request, RequestId, Params, OnComplete);
+	}
+	// resources/list
+	else if (Method == TEXT("resources/list"))
+	{
+		return ProcessListResourcesJsonRpcCall(Request, RequestId, Params, OnComplete);
+	}
+	// resources/read
+	else if (Method == TEXT("resources/read"))
+	{
+		return ProcessReadResourceJsonRpcCall(Request, RequestId, Params, OnComplete);
+	}
+	
+	UE::ModelContextProtocol::Private::CompleteWithError(
+		OnComplete, RequestId, UE::ModelContextProtocol::Private::EJsonRpcErrorCode::MethodNotFound, FString::Printf(TEXT("Call to unknown method \"%s\""), *Method));
+	return true;
+}
+
+bool UModelContextProtocolEngineSubsystem::ProcessPingJsonRpcCall(const FHttpServerRequest& Request, const TSharedPtr<FJsonValue>& RequestId, const FHttpResultCallback& OnComplete)
+{
+	// Reply with an empty result per spec for ping.
+	UE_LOG(LogModelContextProtocol, Log, TEXT("Replying to Ping"));
+	FModelContextProtocolPingResult Result;
+	UE::ModelContextProtocol::Private::CompleteWithResult(OnComplete, RequestId, Result);
+	return true;
+}
+
+bool UModelContextProtocolEngineSubsystem::ProcessInitializeJsonRpcCall(const FHttpServerRequest& Request, const TSharedPtr<FJsonValue>& RequestId, const TSharedPtr<FJsonObject>& Params, const FHttpResultCallback& OnComplete)
+{
+	FModelContextProtocolInitializeParams InitializeParams;
+	if (!UE::ModelContextProtocol::Private::ParseParamsOrError(TEXT("initialize"), RequestId, OnComplete, Params, InitializeParams))
+	{
+		// Request was completed (with error)
+		return true;
+	}
+
+	// Negotiate protocol version per MCP spec:
+	// If the server supports the client's version, respond with that version.
+	// Otherwise, respond with the server's latest supported version.
+	const FString NegotiatedVersion = UE::ModelContextProtocol::NegotiateProtocolVersion(InitializeParams.ProtocolVersion);
+	UE_LOG(LogModelContextProtocol, Display, TEXT("Client requested protocol version '%s', negotiated '%s'"), *InitializeParams.ProtocolVersion, *NegotiatedVersion);
+
+	// Create a new session for this client
+	TSharedPtr<FModelContextProtocolSession> Session = MakeShared<FModelContextProtocolSession>();
+	Session->ID = FGuid::NewGuid().ToString(EGuidFormats::DigitsLower);
+	Session->Status = EModelContextProtocolSessionStatus::Initializing;
+	Session->NegotiatedProtocolVersion = NegotiatedVersion;
+	Session->ClientAddress = Request.PeerAddress;
+	Session->ClientCapabilities = InitializeParams.Capabilities;
+	Sessions.Add(Session);
+	UE_LOG(LogModelContextProtocol, Log, TEXT("Initializing new session: %s"), *Session->ID);
+
+	// Reply with server capabilities
+	FModelContextProtocolInitializeResult InitializeResult;
+	InitializeResult.ProtocolVersion = NegotiatedVersion;
+	FModelContextProtocolToolsCapability ToolsCapability;
+	ToolsCapability.ListChanged = true;
+	InitializeResult.Capabilities.Tools = ToolsCapability;
+	InitializeResult.Capabilities.Resources = FModelContextProtocolResourcesCapability();
+	UE::ModelContextProtocol::Private::CompleteWithResult(OnComplete, RequestId, InitializeResult, &Session->ID);
+	return true;
+}
+
+bool UModelContextProtocolEngineSubsystem::ProcessInitializedNotificationJsonRpcCall(const FHttpServerRequest& Request, const TSharedPtr<FJsonValue>& RequestId, const TSharedPtr<FJsonObject>& Params, const FHttpResultCallback& OnComplete)
+{
+	FString SessionId;
+	if (UE::ModelContextProtocol::Private::GetSessionIdFromRequestHeaders(Request, SessionId))
+	{
+		TSharedPtr<FModelContextProtocolSession> Session = FindSession(SessionId);
+		if (!Session.IsValid())
+		{
+			UE::ModelContextProtocol::Private::CompleteWithError(OnComplete, /*RequestId*/nullptr, UE::ModelContextProtocol::Private::EJsonRpcErrorCode::InvalidRequest, FString::Printf(TEXT("Received notifications/initialized for unknown session: %s"), *SessionId), EHttpServerResponseCodes::BadRequest);
+			return true;
+		}
+		
+		Session->Status = EModelContextProtocolSessionStatus::Initialized;
+		UE_LOG(LogModelContextProtocol, Log, TEXT("Session initialized: %s"), *Session->ID);
+	}
+
+	// Accept notification
+	UE::ModelContextProtocol::Private::CompleteWithResponseCode(OnComplete, EHttpServerResponseCodes::Accepted);
+	return true;
+}
+
+bool UModelContextProtocolEngineSubsystem::ProcessNotificationCancelledJsonRpcCall(const FHttpServerRequest& Request, const TSharedPtr<FJsonValue>& RequestId, const TSharedPtr<FJsonObject>& Params, const FHttpResultCallback& OnComplete)
+{
+	FString SessionId;
+	if (UE::ModelContextProtocol::Private::GetSessionIdFromRequestHeaders(Request, SessionId))
+	{
+		TSharedPtr<FModelContextProtocolSession> Session = FindSession(SessionId);
+		if (!Session.IsValid())
+		{
+			UE::ModelContextProtocol::Private::CompleteWithError(OnComplete, /*RequestId*/nullptr, UE::ModelContextProtocol::Private::EJsonRpcErrorCode::InvalidRequest, FString::Printf(TEXT("Received notifications/cancelled for unknown session: %s"), *SessionId), EHttpServerResponseCodes::BadRequest);
+			return true;
+		}
+
+		if (!Params.IsValid())
+		{
+			UE::ModelContextProtocol::Private::CompleteWithError(OnComplete, /*RequestId*/nullptr, UE::ModelContextProtocol::Private::EJsonRpcErrorCode::InvalidRequest, FString::Printf(TEXT("Expected params for notifications/cancelled for session: %s"), *SessionId), EHttpServerResponseCodes::BadRequest);
+			return true;
+		}
+
+		const TSharedPtr<FJsonValue> CancelId = Params->TryGetField(TEXT("requestId"));
+
+		if (!CancelId.IsValid())
+		{
+			UE::ModelContextProtocol::Private::CompleteWithError(OnComplete, /*RequestId*/nullptr, UE::ModelContextProtocol::Private::EJsonRpcErrorCode::InvalidRequest, FString::Printf(TEXT("Expected requestId for notifications/cancelled for session: %s"), *SessionId), EHttpServerResponseCodes::BadRequest);
+			return true;
+		}
+
+		FModelContextProtocolToolRequestId ToolRequestId(CancelId.ToSharedRef());
+
+		if (const FModelContextProtocolToolContext* Context = Session->ActiveRequests.Find(ToolRequestId))
+		{
+			if (Context->Tool.IsValid())
+			{
+				Context->Tool->CancelAsync(ToolRequestId);
+			}
+
+			Session->ActiveRequests.Remove(ToolRequestId);
+		}
+		else
+		{
+			// Request may have already finished, so not necessarily an error
+		}
+	}
+
+	// Accept notification
+	UE::ModelContextProtocol::Private::CompleteWithResponseCode(OnComplete, EHttpServerResponseCodes::Accepted);
+	return true;
+}
+
+bool UModelContextProtocolEngineSubsystem::ProcessListToolsJsonRpcCall(const FHttpServerRequest& Request, const TSharedPtr<FJsonValue>& RequestId, const TSharedPtr<FJsonObject>& Params, const FHttpResultCallback& OnComplete)
+{
+	FJsonDomBuilder::FArray ToolsArray;
+	TSet<FString> ExistingToolNames;
+
+#if WITH_EDITOR
+	// ToolsetRegistry is the preferred way to expose tools.
+	if (UToolsetRegistrySubsystem* ToolsetRegistrySubsystem = GEditor ? GEditor->GetEditorSubsystem<UToolsetRegistrySubsystem>() : nullptr; ToolsetRegistrySubsystem)
+	{
+		UE::ToolsetRegistry::FToolsetRegistry& ToolsetRegistry = ToolsetRegistrySubsystem->ToolsetRegistry;
+		const FString SchemasJson = ToolsetRegistry.GetToolsetJsonSchemas();
+
+		TArray<TSharedPtr<FJsonValue>> ToolsetArray;
+		TSharedRef<TJsonReader<>> Reader = TJsonReaderFactory<>::Create(SchemasJson);
+		if (FJsonSerializer::Deserialize(Reader, ToolsetArray))
+		{
+			for (const TSharedPtr<FJsonValue>& ToolsetValue : ToolsetArray)
+			{
+				const TSharedPtr<FJsonObject>& ToolsetObject = ToolsetValue->AsObject();
+				if (!ToolsetObject.IsValid())
+				{
+					continue;
+				}
+
+				const TArray<TSharedPtr<FJsonValue>>* ToolsField = nullptr;
+				if (!ToolsetObject->TryGetArrayField(TEXT("tools"), ToolsField) || !ToolsField)
+				{
+					continue;
+				}
+
+				for (const TSharedPtr<FJsonValue>& ToolValue : *ToolsField)
+				{
+					const TSharedPtr<FJsonObject>& ToolJsonObject = ToolValue->AsObject();
+					if (!ToolJsonObject.IsValid())
+					{
+						continue;
+					}
+
+					FString ToolName;
+					ToolJsonObject->TryGetStringField(TEXT("name"), ToolName);
+					if (ToolName.IsEmpty())
+					{
+						continue;
+					}
+
+					ExistingToolNames.Add(ToolName.ToLower());
+
+					FJsonDomBuilder::FObject McpToolObject;
+					McpToolObject.Set(TEXT("name"), ToolName);
+
+					FString ToolDescription;
+					if (ToolJsonObject->TryGetStringField(TEXT("description"), ToolDescription) && !ToolDescription.IsEmpty())
+					{
+						McpToolObject.Set(TEXT("description"), ToolDescription);
+					}
+
+					const TSharedPtr<FJsonObject>* InputSchemaField = nullptr;
+					if (ToolJsonObject->TryGetObjectField(TEXT("inputSchema"), InputSchemaField) && InputSchemaField)
+					{
+						McpToolObject.Set(TEXT("inputSchema"), *InputSchemaField);
+					}
+
+					// Note: ToolsetRegistry tools return plain text via ExecuteTool (FString result),
+					// so we intentionally omit outputSchema to avoid MCP clients expecting structuredContent.
+
+					ToolsArray.Add(McpToolObject);
+				}
+			}
+		}
+	}
+#endif // WITH_EDITOR
+
+	// Deprecated: Also include tools registered directly via IModelContextProtocolModule (UModelContextProtocolToolLibrary).
+	// Prefer ToolsetRegistry for new tool definitions.
+	const IModelContextProtocolModule& ModelContextProtocolModule = IModelContextProtocolModule::GetChecked();
+	for (const TSharedRef<IModelContextProtocolTool>& Tool : ModelContextProtocolModule.GetTools())
+	{
+		const FString ToolName = Tool->GetName();
+		const FString LowercaseToolName = ToolName.ToLower();
+		if (ExistingToolNames.Contains(LowercaseToolName))
+		{
+			UE_LOG(LogModelContextProtocol, Warning, TEXT("Skipping duplicate tool name from deprecated IModelContextProtocolModule: %s"), *ToolName);
+			continue;
+		}
+
+		FJsonDomBuilder::FObject ToolObject;
+		ToolObject.Set(TEXT("name"), ToolName);
+		ExistingToolNames.Add(LowercaseToolName);
+		if (const FString Description = Tool->GetDescription(); !Description.IsEmpty())
+		{
+			ToolObject.Set(TEXT("description"), Description);
+		}
+		if (const TSharedPtr<FJsonObject> InputJsonSchemaObject = Tool->GetInputJsonSchema())
+		{
+			ToolObject.Set(TEXT("inputSchema"), InputJsonSchemaObject);
+		}
+		if (const TSharedPtr<FJsonObject> OutputJsonSchemaObject = Tool->GetOutputJsonSchema())
+		{
+			ToolObject.Set(TEXT("outputSchema"), OutputJsonSchemaObject);
+		}
+		ToolsArray.Add(ToolObject);
+	}
+
+	UE_LOG(LogModelContextProtocol, Log, TEXT("Listing tools (%d)"), ToolsArray.Num());
+
+	TSharedRef<FJsonValueArray> AllToolsJsonValue = ToolsArray.AsJsonValue();
+	const TArray<TSharedPtr<FJsonValue>>& AllTools = AllToolsJsonValue->AsArray();
+	UE::ModelContextProtocol::Private::FPaginationResult PaginationResult;
+	FString PaginationError;
+	if (!UE::ModelContextProtocol::Private::ApplyPagination(AllTools, Params, PaginationResult, PaginationError))
+	{
+		UE::ModelContextProtocol::Private::CompleteWithError(OnComplete, RequestId, UE::ModelContextProtocol::Private::EJsonRpcErrorCode::InvalidParams, PaginationError);
+		return true;
+	}
+
+	FJsonDomBuilder::FObject ResultObject;
+	ResultObject.Set(TEXT("tools"), MakeShared<FJsonValueArray>(MoveTemp(PaginationResult.Items)));
+	if (PaginationResult.NextCursor.IsSet())
+	{
+		ResultObject.Set(TEXT("nextCursor"), PaginationResult.NextCursor.GetValue());
+	}
+
+	UE::ModelContextProtocol::Private::CompleteWithResult(OnComplete, RequestId, ResultObject);
+	return true;
+}
+
+bool UModelContextProtocolEngineSubsystem::ProcessToolCallJsonRpcCall(const FHttpServerRequest& Request, const TSharedPtr<FJsonValue>& RequestId, const TSharedPtr<FJsonObject>& Params, const FHttpResultCallback& OnComplete)
+{
+	if (!Params.IsValid())
+	{
+		UE::ModelContextProtocol::Private::CompleteWithError(OnComplete, RequestId, UE::ModelContextProtocol::Private::EJsonRpcErrorCode::InvalidParams, TEXT("Expected params for call to tools/call"));
+		return true;
+	}
+	
+	FString ToolName;
+	if (!Params->TryGetStringField(TEXT("name"), ToolName) || ToolName.IsEmpty())
+	{
+		UE::ModelContextProtocol::Private::CompleteWithError(OnComplete, RequestId, UE::ModelContextProtocol::Private::EJsonRpcErrorCode::InvalidParams, TEXT("Expected non-empty 'name' param for call to tools/call"));
+		return true;
+	}
+	if (!RequestId.IsValid())
+	{
+		UE::ModelContextProtocol::Private::CompleteWithError(OnComplete, RequestId, UE::ModelContextProtocol::Private::EJsonRpcErrorCode::InvalidParams, TEXT("Invalid request id"));
+		return true;
+	}
+
+	TSharedPtr<FJsonObject> ToolArguments;
+	if (Params->HasTypedField<EJson::Object>(TEXT("arguments")))
+	{
+		ToolArguments = Params->GetObjectField(TEXT("arguments"));
+	}
+
+#if WITH_EDITOR
+	// ToolsetRegistry is the preferred execution path.
+	if (UToolsetRegistrySubsystem* ToolsetRegistrySubsystem = GEditor ? GEditor->GetEditorSubsystem<UToolsetRegistrySubsystem>() : nullptr; ToolsetRegistrySubsystem)
+	{
+		UE::ToolsetRegistry::FToolsetRegistry& ToolsetRegistry = ToolsetRegistrySubsystem->ToolsetRegistry;
+
+		// Check if the tool exists in the ToolsetRegistry before attempting execution.
+		// ExecuteTool splits "ToolsetName.ToolName" and looks up the toolset, so a dotted name
+		// that resolves is a ToolsetRegistry tool.
+		UE::ToolsetRegistry::FToolDescriptor Descriptor;
+		if (auto DescriptorResult = UE::ToolsetRegistry::FToolDescriptor::FromString(ToolName); DescriptorResult.HasValue())
+		{
+			Descriptor = DescriptorResult.GetValue();
+			if (TSharedPtr<UE::ToolsetRegistry::FToolset> Toolset = ToolsetRegistry.Find(Descriptor.ToolsetName))
+			{
+				FString ArgumentsJson;
+				if (ToolArguments.IsValid())
+				{
+					TSharedRef<TJsonWriter<>> Writer = TJsonWriterFactory<>::Create(&ArgumentsJson);
+					FJsonSerializer::Serialize(ToolArguments.ToSharedRef(), Writer);
+				}
+				else
+				{
+					ArgumentsJson = TEXT("{}");
+				}
+
+				UE_LOG(LogModelContextProtocol, Log, TEXT("Running ToolsetRegistry tool: '%s'"), *ToolName);
+
+				ToolsetRegistry.ExecuteTool(Descriptor, ArgumentsJson).Then([RequestId, OnComplete](TFuture<TValueOrError<FString, FString>> Future)
+				{
+					TValueOrError<FString, FString> Result = Future.Get();
+
+					FModelContextProtocolToolResult ToolResult;
+					if (Result.HasError())
+					{
+						ToolResult = UE::ModelContextProtocol::MakeErrorResult(Result.GetError());
+					}
+					else
+					{
+						ToolResult = UE::ModelContextProtocol::MakeTextResult(Result.GetValue());
+					}
+
+					TSharedPtr<FJsonValue> ResultValue = MakeShared<FJsonValueObject>(ToolResult.JsonObject);
+					UE::ModelContextProtocol::Private::CompleteWithResult(OnComplete, RequestId, ResultValue);
+				});
+				return true;
+			}
+		}
+	}
+#endif // WITH_EDITOR
+
+	// Deprecated fallback: try IModelContextProtocolModule (UModelContextProtocolToolLibrary).
+	const IModelContextProtocolModule& ModelContextProtocolModule = IModelContextProtocolModule::GetChecked();
+	TSharedPtr<IModelContextProtocolTool> Tool = ModelContextProtocolModule.FindTool(ToolName);
+	if (!Tool.IsValid())
+	{
+		UE::ModelContextProtocol::Private::CompleteWithError(OnComplete, RequestId, UE::ModelContextProtocol::Private::EJsonRpcErrorCode::InvalidParams, FString::Printf(TEXT("Unknown tool: %s"), *ToolName));
+		return true;
+	}
+
+	UE_LOG(LogModelContextProtocol, Log, TEXT("Running tool: '%s'"), *ToolName);
+	FString PrettyJson;
+	if (UE_LOG_ACTIVE(LogModelContextProtocol, VeryVerbose) && ToolArguments.IsValid() && FJsonSerializer::Serialize(ToolArguments.ToSharedRef(), TJsonWriterFactory<>::Create(&PrettyJson)))
+	{
+		UE_LOG(LogModelContextProtocol, VeryVerbose, TEXT("'%s' parameters:\n%s"),  *ToolName, *PrettyJson);
+	}
+
+	FModelContextProtocolToolRequestId ToolRequestId(RequestId.ToSharedRef());
+
+	FString SessionId;
+	if (UE::ModelContextProtocol::Private::GetSessionIdFromRequestHeaders(Request, SessionId))
+	{
+		TSharedPtr<FModelContextProtocolSession> Session = FindSession(SessionId);
+
+		if (Session.IsValid())
+		{
+			TSharedPtr<FJsonValue> ProgressToken;
+
+			if (Params.IsValid())
+			{
+				const TSharedPtr<FJsonObject>* MetaParams = nullptr;
+
+				if (Params->TryGetObjectField(TEXT("_meta"), MetaParams))
+				{
+					if (MetaParams && MetaParams->IsValid())
+					{
+						ProgressToken = (*MetaParams)->TryGetField(TEXT("progressToken"));
+					}
+				}
+			}
+
+			FModelContextProtocolToolContext& Context = Session->ActiveRequests.Add(ToolRequestId);
+			Context.Tool = Tool;
+			Context.ProgressToken = ProgressToken;
+			// Saving the callback for re-use with SSE event streaming
+			Context.EventStreamWrite = OnComplete;
+			Context.LastProgressSeconds = FPlatformTime::Seconds();
+		}
+		else
+		{
+			UE::ModelContextProtocol::Private::CompleteWithError(OnComplete, RequestId, UE::ModelContextProtocol::Private::EJsonRpcErrorCode::InvalidRequest, FString::Printf(TEXT("Invalid session id")));
+			return true;
+		}
+	}
+	else
+	{
+		UE::ModelContextProtocol::Private::CompleteWithError(OnComplete, RequestId, UE::ModelContextProtocol::Private::EJsonRpcErrorCode::InvalidRequest, FString::Printf(TEXT("Missing session id")));
+		return true;
+	}
+
+	TUniquePtr<FHttpServerResponse> Response = FHttpServerResponse::Create(FString(TEXT("")), UE::ModelContextProtocol::ContentTypeEventStream);
+	Response->Headers.Add(TEXT("Connection"), { TEXT("keep-alive") });
+	Response->Headers.Add(TEXT("Cache-Control"), { TEXT("no-cache") });
+	Response->Headers.Add(UE::ModelContextProtocol::Private::McpSessionIdHeader, { SessionId });
+	EnumAddFlags(Response->Flags, EHttpServerResponseFlags::MultipleWriteStream | EHttpServerResponseFlags::HasAdditionalWrites);
+
+	OnComplete(MoveTemp(Response));
+
+	Tool->RunAsync(ToolRequestId, ToolArguments,
+		/*OnComplete*/[this, RequestId, SessionId, ToolRequestId, OnComplete](const FModelContextProtocolToolResult& Result)
+		{
+			FJsonDomBuilder::FObject ResponseObject;
+			ResponseObject.Set(TEXT("jsonrpc"), UE::ModelContextProtocol::JsonRpcVersion);
+			ResponseObject.Set(TEXT("id"), RequestId);
+
+			TSharedPtr<FModelContextProtocolSession> Session = FindSession(SessionId);
+
+			if (Session.IsValid() && Session->ActiveRequests.Contains(ToolRequestId))
+			{
+				TSharedPtr<FJsonValue> ResultValue = MakeShared<FJsonValueObject>(Result.JsonObject);
+
+				ResponseObject.Set(TEXT("result"), ResultValue);
+
+				Session->ActiveRequests.Remove(ToolRequestId);
+			}
+			else
+			{
+				FJsonDomBuilder::FObject ErrorObject;
+				ErrorObject.Set(TEXT("code"), static_cast<int32>(UE::ModelContextProtocol::Private::EJsonRpcErrorCode::InternalError));
+				ErrorObject.Set(TEXT("message"), TEXT("Unable to find active request."));
+
+				ResponseObject.Set(TEXT("error"), ErrorObject);
+			}
+
+			FUtf8String ResponseStr;
+			TSharedRef<TJsonWriter<UTF8CHAR, TCondensedJsonPrintPolicy<UTF8CHAR>>> JsonWriter = TJsonWriterFactory<UTF8CHAR, TCondensedJsonPrintPolicy<UTF8CHAR>>::Create(&ResponseStr);
+			FJsonSerializer::Serialize(ResponseObject.AsJsonObject(), JsonWriter);
+
+			TUniquePtr<FHttpServerResponse> ServerResponse = FHttpServerResponse::Create(UE::ModelContextProtocol::Private::FormatSSEMessage(ResponseStr), UE::ModelContextProtocol::ContentTypeEventStream);
+			EnumAddFlags(ServerResponse->Flags, EHttpServerResponseFlags::MultipleWriteStream | EHttpServerResponseFlags::SkipHeaderWrite);
+
+			OnComplete(MoveTemp(ServerResponse));
+		});
+
+	return true;
+}
+
+bool UModelContextProtocolEngineSubsystem::ProcessListResourcesJsonRpcCall(const FHttpServerRequest& Request, const TSharedPtr<FJsonValue>& RequestId, const TSharedPtr<FJsonObject>& Params, const FHttpResultCallback& OnComplete)
+{
+	const IModelContextProtocolModule& ModelContextProtocolModule = IModelContextProtocolModule::GetChecked();
+
+	LastResourceDescriptorList.Reset();
+	for (const TSharedRef<IModelContextProtocolResourceProvider>& ResourceProvider : ModelContextProtocolModule.GetResourceProviders())
+	{
+		ResourceProvider->ListResources(LastResourceDescriptorList);
+	}
+
+	UE_LOG(LogModelContextProtocol, Log, TEXT("Listing resources (%d)"), LastResourceDescriptorList.Num());
+
+	TSharedRef<FJsonValueArray> AllResourcesJsonValue = LastResourceDescriptorList.GetJsonArray();
+	const TArray<TSharedPtr<FJsonValue>>& AllResources = AllResourcesJsonValue->AsArray();
+	UE::ModelContextProtocol::Private::FPaginationResult PaginationResult;
+	FString PaginationError;
+	if (!UE::ModelContextProtocol::Private::ApplyPagination(AllResources, Params, PaginationResult, PaginationError))
+	{
+		UE::ModelContextProtocol::Private::CompleteWithError(OnComplete, RequestId, UE::ModelContextProtocol::Private::EJsonRpcErrorCode::InvalidParams, PaginationError);
+		LastResourceDescriptorList.ReleaseJsonArray();
+		return true;
+	}
+
+	FJsonDomBuilder::FObject ResultObject;
+	ResultObject.Set(TEXT("resources"), MakeShared<FJsonValueArray>(MoveTemp(PaginationResult.Items)));
+	if (PaginationResult.NextCursor.IsSet())
+	{
+		ResultObject.Set(TEXT("nextCursor"), PaginationResult.NextCursor.GetValue());
+	}
+
+	UE::ModelContextProtocol::Private::CompleteWithResult(OnComplete, RequestId, ResultObject);
+
+	LastResourceDescriptorList.ReleaseJsonArray();
+
+	return true;
+}
+
+bool UModelContextProtocolEngineSubsystem::ProcessReadResourceJsonRpcCall(const FHttpServerRequest& Request, const TSharedPtr<FJsonValue>& RequestId, const TSharedPtr<FJsonObject>& Params, const FHttpResultCallback& OnComplete)
+{
+	if (!Params.IsValid())
+	{
+		UE::ModelContextProtocol::Private::CompleteWithError(OnComplete, RequestId, UE::ModelContextProtocol::Private::EJsonRpcErrorCode::InvalidParams, TEXT("Expected params for call to resources/read"));
+		return true;
+	}
+	
+	FString ResourceUri;
+	if (!Params->TryGetStringField(TEXT("uri"), ResourceUri))
+	{
+		UE::ModelContextProtocol::Private::CompleteWithError(OnComplete, RequestId, UE::ModelContextProtocol::Private::EJsonRpcErrorCode::InvalidParams, TEXT("Expected uri param for call to resources/read"));
+		return true;
+	}
+
+	// Look up which resource provider last listed this resource
+	const TSharedPtr<const IModelContextProtocolResourceProvider> ResourceProvider = LastResourceDescriptorList.FindResourceProvider(ResourceUri);
+	if (!ResourceProvider.IsValid())
+	{
+		UE::ModelContextProtocol::Private::CompleteWithError(OnComplete, RequestId, UE::ModelContextProtocol::Private::EJsonRpcErrorCode::ResourceNotFound, FString::Printf(TEXT("Resource not found: %s"), *ResourceUri));
+		return true;
+	}
+
+	TValueOrError<FModelContextProtocolResource, FString> ResourceOrError = ResourceProvider->ReadResource(ResourceUri);
+	if (ResourceOrError.HasError())
+	{
+		UE::ModelContextProtocol::Private::CompleteWithError(OnComplete, RequestId, UE::ModelContextProtocol::Private::EJsonRpcErrorCode::InternalError, FString::Printf(TEXT("Error reading resource: %s"), *ResourceOrError.GetError()));
+		return true;
+	}
+
+	FJsonDomBuilder::FObject ResultObject;
+	ResultObject.Set(TEXT("contents"), FJsonDomBuilder::FArray().Add(ResourceOrError.GetValue().GetJsonObject()));
+
+	UE::ModelContextProtocol::Private::CompleteWithResult(OnComplete, RequestId, ResultObject);
+	return true;
+}
+
+void UModelContextProtocolEngineSubsystem::Tick(float DeltaTime)
+{
+	if (UE::ModelContextProtocol::ProgressIntervalSeconds > 0.0f)
+	{
+		const double NowSeconds = FPlatformTime::Seconds();
+
+		for (TSharedPtr<FModelContextProtocolSession>& Session : Sessions)
+		{
+			if (Session.IsValid())
+			{
+				for (TPair<FModelContextProtocolToolRequestId, FModelContextProtocolToolContext>& ToolRequest : Session->ActiveRequests)
+				{
+					FModelContextProtocolToolContext& Context = ToolRequest.Value;
+
+					if (Context.ProgressToken.IsValid() && Context.EventStreamWrite.IsSet())
+					{
+						if ((NowSeconds - Context.LastProgressSeconds) >= UE::ModelContextProtocol::ProgressIntervalSeconds)
+						{
+							// Not sending a total progress value, so this is simply a heartbeat as the total duration is unknown
+							++Context.LastProgressValue;
+
+							UE::ModelContextProtocol::Private::SendProgressUpdate(Context.EventStreamWrite, Context.ProgressToken.ToSharedRef(), Context.LastProgressValue);
+
+							Context.LastProgressSeconds = NowSeconds;
+						}
+					}
+				}
+			}
+		}
+	}
+}
+
+bool UModelContextProtocolEngineSubsystem::ProcessGetRequest(const FHttpServerRequest& Request, const FHttpResultCallback& OnComplete)
+{
+	if (!UE::ModelContextProtocol::Private::ValidateOriginHeader(Request, OnComplete))
+	{
+		return true;
+	}
+
+	// We do not currently support sse on a separate endpoint
+	UE::ModelContextProtocol::Private::CompleteWithResponseCode(OnComplete, EHttpServerResponseCodes::BadMethod);
+	return true;
+}
+
+bool UModelContextProtocolEngineSubsystem::ProcessDeleteRequest(const FHttpServerRequest& Request, const FHttpResultCallback& OnComplete)
+{
+	if (!UE::ModelContextProtocol::Private::ValidateOriginHeader(Request, OnComplete))
+	{
+		return true;
+	}
+
+	FString SessionId;
+	if (UE::ModelContextProtocol::Private::GetSessionIdFromRequestHeaders(Request, SessionId))
+	{
+		const int32 RemovedCount = Sessions.RemoveAll([&SessionId](const TSharedPtr<FModelContextProtocolSession>& Session)
+			{
+				return Session.IsValid() && (Session->ID == SessionId);
+			});
+
+		if (RemovedCount > 0)
+		{
+			UE::ModelContextProtocol::Private::CompleteWithResponseCode(OnComplete, EHttpServerResponseCodes::Accepted);
+		}
+		else
+		{
+			UE::ModelContextProtocol::Private::CompleteWithResponseCode(OnComplete, EHttpServerResponseCodes::BadRequest);
+		}
+	}
+	else
+	{
+		UE::ModelContextProtocol::Private::CompleteWithResponseCode(OnComplete, EHttpServerResponseCodes::BadRequest);
+	}
+
+	return true;
+}
diff --git a/Engine/Plugins/Experimental/ModelContextProtocol/Source/ModelContextProtocol/Private/ModelContextProtocolMetaData.cpp b/Engine/Plugins/Experimental/ModelContextProtocol/Source/ModelContextProtocol/Private/ModelContextProtocolMetaData.cpp
new file mode 100644
--- /dev/null
+++ b/Engine/Plugins/Experimental/ModelContextProtocol/Source/ModelContextProtocol/Private/ModelContextProtocolMetaData.cpp
@@ -0,0 +1,206 @@
+// Copyright Epic Games, Inc. All Rights Reserved.
+
+#include "ModelContextProtocolMetaData.h"
+
+#include "JsonSchema/JsonSchemaEditorMetadata.h"
+#include "JsonSchema/JsonSchemaPropertyFilter.h"
+#include "Serialization/JsonSerializer.h"
+
+#if WITH_EDITORONLY_DATA
+#include "EdGraphSchema_K2.h"
+#include "JsonSchema/JsonSchemaGeneratorEditor.h"
+#endif // WITH_EDITORONLY_DATA
+
+namespace UE::ModelContextProtocol
+{
+
+#if WITH_EDITORONLY_DATA
+FModelContextProtocolFunctionMetaData CollectFunctionMetaData(const UFunction* Function, const FJsonSchemaPropertyFilter& PropertyFilter)
+{
+	FModelContextProtocolFunctionMetaData Result;
+
+	if (!Function)
+	{
+		return Result;
+	}
+
+	FJsonSchemaEditorMetadata EditorMetadata = FJsonSchemaGeneratorEditor::UStructToJsonSchemaMetadata(Function, PropertyFilter, nullptr);
+
+	// Extract function description from root struct metadata
+	if (EditorMetadata.RootStructMetadata.IsValid())
+	{
+		FString DescriptionString;
+		if (EditorMetadata.RootStructMetadata->TryGetStringField(TEXT("description"), DescriptionString) && !DescriptionString.IsEmpty())
+		{
+			Result.Description = FText::FromString(DescriptionString);
+		}
+	}
+
+	// Extract WorldContext directly from UFunction metadata (FJsonSchemaGeneratorEditor does not collect this)
+	if (Function->HasMetaData(FBlueprintMetadata::MD_WorldContext))
+	{
+		Result.WorldContext = Function->GetMetaData(FBlueprintMetadata::MD_WorldContext);
+	}
+
+	// Extract per-property metadata
+	for (const auto& [PropertyPath, PropertyMetadataJsonObject] : EditorMetadata.GetPropertyMemberPathToPropertyMetadataMap())
+	{
+		if (!PropertyMetadataJsonObject.IsValid())
+		{
+			continue;
+		}
+
+		FModelContextProtocolPropertyMetaData PropertyMetaData;
+
+		// Description
+		FString DescriptionString;
+		if (PropertyMetadataJsonObject->TryGetStringField(TEXT("description"), DescriptionString) && !DescriptionString.IsEmpty())
+		{
+			PropertyMetaData.Description = FText::FromString(DescriptionString);
+		}
+
+		// Default value - serialize back to string for cookable storage, preserving the original JSON type
+		if (PropertyMetadataJsonObject->HasField(TEXT("default")))
+		{
+			const TSharedPtr<FJsonValue>& DefaultValue = PropertyMetadataJsonObject->Values.FindChecked(TEXT("default"));
+			if (DefaultValue.IsValid())
+			{
+				PropertyMetaData.DefaultValueJsonType = static_cast<uint8>(DefaultValue->Type);
+				switch (DefaultValue->Type)
+				{
+				case EJson::String:
+					PropertyMetaData.DefaultValue = DefaultValue->AsString();
+					break;
+				case EJson::Number:
+					PropertyMetaData.DefaultValue = FString::SanitizeFloat(DefaultValue->AsNumber());
+					break;
+				case EJson::Boolean:
+					PropertyMetaData.DefaultValue = DefaultValue->AsBool() ? TEXT("true") : TEXT("false");
+					break;
+				case EJson::Null:
+					PropertyMetaData.DefaultValue = TEXT("None");
+					break;
+				case EJson::Object:
+				case EJson::Array:
+					{
+						FString JsonString;
+						TSharedRef<TJsonWriter<>> Writer = TJsonWriterFactory<>::Create(&JsonString);
+						FJsonSerializer::Serialize(DefaultValue, TEXT(""), Writer);
+						PropertyMetaData.DefaultValue = JsonString;
+					}
+					break;
+				default:
+					break;
+				}
+			}
+		}
+
+		// ClampMin (stored as "minimum")
+		double MinValue;
+		if (PropertyMetadataJsonObject->TryGetNumberField(TEXT("minimum"), MinValue))
+		{
+			PropertyMetaData.ClampMin = MinValue;
+		}
+
+		// ClampMax (stored as "maximum")
+		double MaxValue;
+		if (PropertyMetadataJsonObject->TryGetNumberField(TEXT("maximum"), MaxValue))
+		{
+			PropertyMetaData.ClampMax = MaxValue;
+		}
+
+		Result.PropertyMetaData.Add(PropertyPath, MoveTemp(PropertyMetaData));
+	}
+
+	return Result;
+}
+#endif // WITH_EDITORONLY_DATA
+
+FJsonSchemaEditorMetadata ConvertToCachedEditorMetadata(const FModelContextProtocolFunctionMetaData& FunctionMetaData)
+{
+	FJsonSchemaEditorMetadata EditorMetadata;
+
+	// Set root struct description
+	if (FunctionMetaData.Description.IsSet())
+	{
+		EditorMetadata.RootStructMetadata = MakeShared<FJsonObject>();
+		EditorMetadata.RootStructMetadata->SetStringField(TEXT("description"), FunctionMetaData.Description->ToString());
+	}
+
+	// Convert per-property metadata
+	FJsonSchemaEditorMetadata::FJsonSchemaPropertyMemberPathToPropertyMetadataMap& PropertyMap = EditorMetadata.GetPropertyMemberPathToPropertyMetadataMap();
+	for (const auto& [PropertyPath, PropertyMetaData] : FunctionMetaData.PropertyMetaData)
+	{
+		TSharedPtr<FJsonObject> PropertyJsonObject = MakeShared<FJsonObject>();
+
+		if (PropertyMetaData.Description.IsSet())
+		{
+			PropertyJsonObject->SetStringField(TEXT("description"), PropertyMetaData.Description->ToString());
+		}
+
+		if (PropertyMetaData.DefaultValue.IsSet())
+		{
+			const FString& DefaultString = *PropertyMetaData.DefaultValue;
+			const EJson OriginalType = PropertyMetaData.DefaultValueJsonType.IsSet()
+				? static_cast<EJson>(*PropertyMetaData.DefaultValueJsonType)
+				: EJson::String;
+
+			switch (OriginalType)
+			{
+			case EJson::Boolean:
+				PropertyJsonObject->SetBoolField(TEXT("default"), DefaultString == TEXT("true"));
+				break;
+			case EJson::Null:
+				PropertyJsonObject->SetField(TEXT("default"), MakeShared<FJsonValueNull>());
+				break;
+			case EJson::Number:
+				{
+					double NumericValue;
+					if (LexTryParseString(NumericValue, *DefaultString))
+					{
+						PropertyJsonObject->SetNumberField(TEXT("default"), NumericValue);
+					}
+					else
+					{
+						PropertyJsonObject->SetStringField(TEXT("default"), DefaultString);
+					}
+				}
+				break;
+			case EJson::Object:
+			case EJson::Array:
+				{
+					TSharedPtr<FJsonValue> ParsedValue;
+					TSharedRef<TJsonReader<>> Reader = TJsonReaderFactory<>::Create(DefaultString);
+					if (FJsonSerializer::Deserialize(Reader, ParsedValue) && ParsedValue.IsValid())
+					{
+						PropertyJsonObject->SetField(TEXT("default"), ParsedValue);
+					}
+					else
+					{
+						PropertyJsonObject->SetStringField(TEXT("default"), DefaultString);
+					}
+				}
+				break;
+			default:
+				PropertyJsonObject->SetStringField(TEXT("default"), DefaultString);
+				break;
+			}
+		}
+
+		if (PropertyMetaData.ClampMin.IsSet())
+		{
+			PropertyJsonObject->SetNumberField(TEXT("minimum"), *PropertyMetaData.ClampMin);
+		}
+
+		if (PropertyMetaData.ClampMax.IsSet())
+		{
+			PropertyJsonObject->SetNumberField(TEXT("maximum"), *PropertyMetaData.ClampMax);
+		}
+
+		PropertyMap.Add(PropertyPath, PropertyJsonObject);
+	}
+
+	return EditorMetadata;
+}
+
+}
diff --git a/Engine/Plugins/Experimental/ModelContextProtocol/Source/ModelContextProtocol/Public/ModelContextProtocolEngineSubsystem.h b/Engine/Plugins/Experimental/ModelContextProtocol/Source/ModelContextProtocol/Public/ModelContextProtocolEngineSubsystem.h
new file mode 100644
--- /dev/null
+++ b/Engine/Plugins/Experimental/ModelContextProtocol/Source/ModelContextProtocol/Public/ModelContextProtocolEngineSubsystem.h
@@ -0,0 +1,97 @@
+﻿// Copyright Epic Games, Inc. All Rights Reserved.
+
+#pragma once
+
+#include "ModelContextProtocolResources.h"
+#include "ModelContextProtocolSession.h"
+
+#include "HttpResultCallback.h"
+#include "HttpRouteHandle.h"
+#include "Subsystems/EngineSubsystem.h"
+#include "Templates/SharedPointer.h"
+#include "Tickable.h"
+
+#include "ModelContextProtocolEngineSubsystem.generated.h"
+
+class IHttpRouter;
+class FJsonValue;
+class FJsonObject;
+
+/**
+ * MCP (Anthropic's Model Context Protocol) tool server. Serves MCP tools over HTTP per the https://modelcontextprotocol.io/ spec.  
+ * @see https://modelcontextprotocol.io
+ * @see UModelContextProtocolSettings
+ */
+UCLASS(MinimalAPI)
+class UModelContextProtocolEngineSubsystem : public UEngineSubsystem, public FTickableGameObject
+{
+	GENERATED_BODY()
+public:
+
+	/**
+	 * Starts the HTTP MCP server on the configured port.
+	 *
+	 * Called automatically on subsystem Initialize if UModelContextProtocolSettings::bAutoStartServer is enabled.
+	 */
+	MODELCONTEXTPROTOCOL_API void StartServer();
+
+	/**
+	 * Stops the HTTP MCP server.
+	 */
+	MODELCONTEXTPROTOCOL_API void StopServer();
+
+	/** Returns true if the HTTP MCP server is currently running. */
+	MODELCONTEXTPROTOCOL_API bool IsServerRunning() const;
+
+	/** Returns the port the server is currently listening on, or 0 if not running. */
+	MODELCONTEXTPROTOCOL_API uint32 GetServerPort() const;
+
+	/** Sets the server port number in the user settings. Takes effect on next StartServer() call. */
+	MODELCONTEXTPROTOCOL_API void SetServerPort(uint32 Port);
+
+	//~ Begin USubsystem API
+	virtual void Initialize(FSubsystemCollectionBase& Collection) override;
+	virtual void Deinitialize() override;
+	//~ End USubsystem API
+
+	//~ Begin FTickableGameObject interface
+	virtual TStatId GetStatId() const override { RETURN_QUICK_DECLARE_CYCLE_STAT(UModelContextProtocolEngineSubsystem, STATGROUP_Tickables); }
+	virtual ETickableTickType GetTickableTickType() const override { return (IsTemplate() ? ETickableTickType::Never : ETickableTickType::Always); }
+	virtual bool IsTickableInEditor() const override { return true; }
+	virtual void Tick(float DeltaTime) override;
+	//~ End FTickableGameObject interface
+
+protected:
+
+	TArray<TSharedPtr<FModelContextProtocolSession>> Sessions;
+
+	TSharedPtr<FModelContextProtocolSession> FindSession(const FString& SessionId) const;
+
+	/** Sends notifications/tools/list_changed to all initialized sessions with active SSE streams. */
+	void BroadcastToolsListChanged();
+
+	FDelegateHandle OnRefreshToolsDelegateHandle;
+
+	TSharedPtr<IHttpRouter> HttpRouter;
+	uint32 ActiveServerPort = 0;
+
+	FHttpRouteHandle MainMcpRoute;
+	FHttpRouteHandle SseMcpRoute;
+	FHttpRouteHandle DeleteMcpRoute;
+
+	FModelContextProtocolResourceDescriptorList LastResourceDescriptorList;
+
+	bool ProcessPostRequest(const FHttpServerRequest& Request, const FHttpResultCallback& OnComplete);
+	bool ProcessGetRequest(const FHttpServerRequest& Request, const FHttpResultCallback& OnComplete);
+	bool ProcessDeleteRequest(const FHttpServerRequest& Request, const FHttpResultCallback& OnComplete);
+
+	bool ProcessJsonRpcCall(const FHttpServerRequest& Request, const TSharedPtr<FJsonValue>& RequestId, const FString& Method, const TSharedPtr<FJsonObject>& Params, const FHttpResultCallback& OnComplete);
+	bool ProcessPingJsonRpcCall(const FHttpServerRequest& Request, const TSharedPtr<FJsonValue>& RequestId, const FHttpResultCallback& OnComplete);
+	bool ProcessInitializeJsonRpcCall(const FHttpServerRequest& Request, const TSharedPtr<FJsonValue>& RequestId, const TSharedPtr<FJsonObject>& Params, const FHttpResultCallback& OnComplete);
+	bool ProcessInitializedNotificationJsonRpcCall(const FHttpServerRequest& Request, const TSharedPtr<FJsonValue>& RequestId, const TSharedPtr<FJsonObject>& Params, const FHttpResultCallback& OnComplete);
+	bool ProcessNotificationCancelledJsonRpcCall(const FHttpServerRequest& Request, const TSharedPtr<FJsonValue>& RequestId, const TSharedPtr<FJsonObject>& Params, const FHttpResultCallback& OnComplete);
+	bool ProcessListToolsJsonRpcCall(const FHttpServerRequest& Request, const TSharedPtr<FJsonValue>& RequestId, const TSharedPtr<FJsonObject>& Params, const FHttpResultCallback& OnComplete);
+	bool ProcessToolCallJsonRpcCall(const FHttpServerRequest& Request, const TSharedPtr<FJsonValue>& RequestId, const TSharedPtr<FJsonObject>& Params, const FHttpResultCallback& OnComplete);
+	bool ProcessListResourcesJsonRpcCall(const FHttpServerRequest& Request, const TSharedPtr<FJsonValue>& RequestId, const TSharedPtr<FJsonObject>& Params, const FHttpResultCallback& OnComplete);
+	bool ProcessReadResourceJsonRpcCall(const FHttpServerRequest& Request, const TSharedPtr<FJsonValue>& RequestId, const TSharedPtr<FJsonObject>& Params, const FHttpResultCallback& OnComplete);
+};
diff --git a/Engine/Plugins/Experimental/ModelContextProtocol/Source/ModelContextProtocol/Public/ModelContextProtocolMetaData.h b/Engine/Plugins/Experimental/ModelContextProtocol/Source/ModelContextProtocol/Public/ModelContextProtocolMetaData.h
new file mode 100644
--- /dev/null
+++ b/Engine/Plugins/Experimental/ModelContextProtocol/Source/ModelContextProtocol/Public/ModelContextProtocolMetaData.h
@@ -0,0 +1,79 @@
+// Copyright Epic Games, Inc. All Rights Reserved.
+
+#pragma once
+
+#include "Internationalization/Text.h"
+#include "Misc/Optional.h"
+
+#include "ModelContextProtocolMetaData.generated.h"
+
+struct FJsonSchemaEditorMetadata;
+struct FJsonSchemaPropertyFilter;
+
+/** Cached editor-only property meta-data in a cookable form */
+USTRUCT()
+struct FModelContextProtocolPropertyMetaData
+{
+	GENERATED_BODY()
+
+	/** Property description derived from its tooltip / native code comment */
+	UPROPERTY()
+	TOptional<FText> Description;
+
+	/** Property's Meta ClampMin e.g: UPROPERTY(Meta = (ClampMin = "1")) */
+	UPROPERTY()
+	TOptional<double> ClampMin;
+
+	/** Property's Meta ClampMax e.g: UPROPERTY(Meta = (ClampMax = "1000")) */
+	UPROPERTY()
+	TOptional<double> ClampMax;
+
+	/** For function parameters, stores the default parameter value */
+	UPROPERTY()
+	TOptional<FString> DefaultValue;
+
+	/** Original JSON type (EJson cast to uint8) of DefaultValue, used for lossless round-trip */
+	UPROPERTY()
+	TOptional<uint8> DefaultValueJsonType;
+};
+
+/** Cached editor-only function meta-data in a cookable form */
+USTRUCT()
+struct FModelContextProtocolFunctionMetaData
+{
+	GENERATED_BODY()
+
+	/** Function description derived from its tooltip / native code comment */
+	UPROPERTY()
+	TOptional<FText> Description;
+
+	/** Per-property meta data, including inner properties e.g: Person, Person.Name, Person.Age */
+	UPROPERTY()
+	TMap<FString, FModelContextProtocolPropertyMetaData> PropertyMetaData;
+
+	/** The function's WorldContext meta-data value, if any */
+	UPROPERTY()
+	FString WorldContext;
+};
+
+namespace UE::ModelContextProtocol
+{
+#if WITH_EDITORONLY_DATA
+	/**
+	 * Collects editor-only metadata for a UFunction, storing it in cookable form.
+	 *
+	 * @param Function The function to collect metadata for.
+	 * @param PropertyFilter Property filter for schema generation.
+	 * @return Collected function metadata.
+	 */
+	MODELCONTEXTPROTOCOL_API FModelContextProtocolFunctionMetaData CollectFunctionMetaData(const UFunction* Function, const FJsonSchemaPropertyFilter& PropertyFilter);
+#endif
+
+	/**
+	 * Converts cookable function metadata back into FJsonSchemaEditorMetadata for passing to FJsonSchemaGenerator.
+	 *
+	 * @param FunctionMetaData The cookable metadata to convert.
+	 * @return Engine editor metadata suitable for schema generation.
+	 */
+	MODELCONTEXTPROTOCOL_API FJsonSchemaEditorMetadata ConvertToCachedEditorMetadata(const FModelContextProtocolFunctionMetaData& FunctionMetaData);
+}
diff --git a/Engine/Plugins/Experimental/ModelContextProtocol/Source/ModelContextProtocol/Public/ModelContextProtocolSettings.h b/Engine/Plugins/Experimental/ModelContextProtocol/Source/ModelContextProtocol/Public/ModelContextProtocolSettings.h
new file mode 100644
--- /dev/null
+++ b/Engine/Plugins/Experimental/ModelContextProtocol/Source/ModelContextProtocol/Public/ModelContextProtocolSettings.h
@@ -0,0 +1,37 @@
+﻿// Copyright Epic Games, Inc. All Rights Reserved.
+
+#pragma once
+
+#include "Engine/DeveloperSettings.h"
+
+#include "ModelContextProtocolSettings.generated.h"
+
+namespace UE::ModelContextProtocol
+{
+	bool ShouldAutoStartServer();
+	uint32 GetServerPortNumber();
+	FString GetServerUrlPath();
+}
+
+UCLASS(MinimalAPI, config=EditorPerProjectUserSettings, defaultconfig, meta=(DisplayName="Model Context Protocol"))
+class UModelContextProtocolSettings : public UDeveloperSettings
+{
+	GENERATED_BODY()
+public:
+
+	/** The URL base path to serve from e.g: "/mcp" -> http://localhost/mcp */
+	UPROPERTY(Config, EditAnywhere, Category="Server")
+	FString ServerUrlPath = TEXT("/mcp");
+	
+	/** The port number to serve from e.g: 8000 -> http://localhost:8000/mcp */
+	UPROPERTY(Config, EditAnywhere, Category="Server")
+	uint32 ServerPortNumber = 8000;
+
+	/**
+	 * If true, the HTTP server route will be automatically registered and HTTP listeners started, during UModelContextProtocolEngineSubsystem
+	 * initialization.
+	 */
+	UPROPERTY(Config, EditAnywhere, Category="Server")
+	bool bAutoStartServer = false;
+};
+
