[ModelContextProtocol] Decouple *MCP* server from Engine, split into three-module architecture.

Make `ModelContextProtocol` module completely Engine-free so a standalone application can use the *MCP* HTTP server without depending on Engine.

Introduce new three modules structure:

- `ModelContextProtocol` (Runtime, no Engine): HTTP server (`FModelContextProtocolServer`), `IModelContextProtocolTool` interface, `IModelContextProtocolResourceProvider`, tool results, sessions, protocol constants.
- `ModelContextProtocolEngine` (Runtime, depends on Engine): `UModelContextProtocolSettings`, `UModelContextProtocolToolLibrary`, `UModelContextProtocolToolAsyncAction`, `UTexture2D`/`USoundWave` result helpers, `ModelContextProtocolToolUtils`.
- `ModelContextProtocolEditor` (Editor): `FToolsetRegistryToolAdapter` (bridges ToolsetRegistry tools into `IModelContextProtocolTool`), asset definitions, editor tool library, factories.

# Changes

- Delete `UModelContextProtocolEngineSubsystem`: replaced by `FModelContextProtocolServer` (plain C++ class with `FTSTicker` for heartbeats instead of `FTickableGameObject`).
- Server ownership moved to `IModelContextProtocolModule` / `FModelContextProtocolModule` with `StartServer(Port, UrlPath)` parameterized API.
- ToolsetRegistry integration moved from direct server queries to `FToolsetRegistryToolAdapter` pattern: the server only sees `IModelContextProtocolTool`, zero editor awareness.
- Editor module defers initialization to `FCoreDelegates::GetOnPostEngineInit()` and listens to `ToolsetRegistry.OnToolsetRegistered()` for dynamic tool registration.
- Remove `Engine` and `DeveloperSettings` from `ModelContextProtocol.Build.cs`: module depends only on `Core`, `CoreUObject`, `HTTPServer`, `Json`, `JsonUtilities`.
- Remove unused `Engine/EngineBaseTypes.h` includes from `IModelContextProtocolResourceProvider.h` and `ModelContextProtocolResources.h`.
- Export CVars (`bWrapPODResultsInObject`, `bAudioResultOggFormat`, `ProgressIntervalSeconds`) for cross-module access.
- Add `bAlive` weak pointer guard in `FModelContextProtocolServer` to prevent use-after-free when async tool calls complete after server destruction.
- Fix IPv6 localhost origin validation: bracketed addresses like `[::1]:8080` are now parsed correctly in the DNS rebinding check.
- Validate console command `ModelContextProtocol.StartServer` port argument (1-65535 range check).
- Fix `resources/read` failing without prior `resources/list`: on cache miss, query all providers directly per MCP spec.
- Move ToolsetRegistry tests and `MockToolsetDefinition` to `ModelContextProtocolEditorTests`.
- Move `ModelContextProtocolToolUtilsTests` to `ModelContextProtocolEditorTests`.
- Clean up dead `Subsystem` references in base tests, rename spec to `AI.ModelContextProtocol.Server`.

# Features

- Add `ValidateToolName()` with `EToolNameValidation` enum per MCP spec ([REDACTED_PHONE]) tool name rules; rejects empty names, warns on length violations, logs invalid characters.
- Add `AI.ModelContextProtocol.ToolNames` test spec.

#Synced-CL 52453059

diff --git a/Engine/Plugins/Experimental/ModelContextProtocol/ModelContextProtocol.uplugin b/Engine/Plugins/Experimental/ModelContextProtocol/ModelContextProtocol.uplugin
--- a/Engine/Plugins/Experimental/ModelContextProtocol/ModelContextProtocol.uplugin
+++ b/Engine/Plugins/Experimental/ModelContextProtocol/ModelContextProtocol.uplugin
@@ -22,6 +22,11 @@
 			"Type": "Runtime",
 			"LoadingPhase": "Default"
 		},
+		{
+			"Name": "ModelContextProtocolEngine",
+			"Type": "Runtime",
+			"LoadingPhase": "Default"
+		},
 		{
 			"Name": "ModelContextProtocolEditor",
 			"Type": "Editor",
diff --git a/Engine/Plugins/Experimental/ModelContextProtocol/Source/ModelContextProtocol/ModelContextProtocol.Build.cs b/Engine/Plugins/Experimental/ModelContextProtocol/Source/ModelContextProtocol/ModelContextProtocol.Build.cs
--- a/Engine/Plugins/Experimental/ModelContextProtocol/Source/ModelContextProtocol/ModelContextProtocol.Build.cs
+++ b/Engine/Plugins/Experimental/ModelContextProtocol/Source/ModelContextProtocol/ModelContextProtocol.Build.cs
@@ -7,12 +7,11 @@ public class ModelContextProtocol : ModuleRules
 	public ModelContextProtocol(ReadOnlyTargetRules Target) : base(Target)
 	{
 		PCHUsage = ModuleRules.PCHUsageMode.UseExplicitOrSharedPCHs;
-		
+
 		PublicDependencyModuleNames.AddRange(
 			new string[]
 			{
 				"Core",
-				"Engine",
 				"HTTPServer",
 				"JsonUtilities",
 				"Json"
@@ -21,21 +20,7 @@ public class ModelContextProtocol : ModuleRules
 		PrivateDependencyModuleNames.AddRange(
 			new string[]
 			{
-				"CoreUObject",
-				"DeveloperSettings"
-			});
-
-		if (Target.bBuildEditor)
-		{
-			PrivateDependencyModuleNames.AddRange(
-			new string[]
-			{
-				"BlueprintGraph",
-				"JsonUtilitiesEditor",
-				"KismetCompiler",
-				"ToolsetRegistry",
-				"UnrealEd"
+				"CoreUObject"
 			});
-		}
 	}
 }
diff --git a/Engine/Plugins/Experimental/ModelContextProtocol/Source/ModelContextProtocol/Private/ModelContextProtocol.cpp b/Engine/Plugins/Experimental/ModelContextProtocol/Source/ModelContextProtocol/Private/ModelContextProtocol.cpp
--- a/Engine/Plugins/Experimental/ModelContextProtocol/Source/ModelContextProtocol/Private/ModelContextProtocol.cpp
+++ b/Engine/Plugins/Experimental/ModelContextProtocol/Source/ModelContextProtocol/Private/ModelContextProtocol.cpp
@@ -6,6 +6,37 @@
 
 DEFINE_LOG_CATEGORY(LogModelContextProtocol);
 
+UE::ModelContextProtocol::EToolNameValidation UE::ModelContextProtocol::ValidateToolName(const FString& ToolName)
+{
+	if (ToolName.IsEmpty())
+	{
+		return EToolNameValidation::Empty;
+	}
+
+	if (ToolName.Len() > 128)
+	{
+		return EToolNameValidation::ExceedsMaxLength;
+	}
+
+	for (const TCHAR Character : ToolName)
+	{
+		const bool bIsAllowed =
+			(Character >= TEXT('A') && Character <= TEXT('Z')) ||
+			(Character >= TEXT('a') && Character <= TEXT('z')) ||
+			(Character >= TEXT('0') && Character <= TEXT('9')) ||
+			Character == TEXT('_') ||
+			Character == TEXT('-') ||
+			Character == TEXT('.');
+
+		if (!bIsAllowed)
+		{
+			return EToolNameValidation::InvalidCharacters;
+		}
+	}
+
+	return EToolNameValidation::Valid;
+}
+
 // CVars
 namespace UE::ModelContextProtocol
 {
diff --git a/Engine/Plugins/Experimental/ModelContextProtocol/Source/ModelContextProtocol/Private/ModelContextProtocolEngineSubsystem.cpp b/Engine/Plugins/Experimental/ModelContextProtocol/Source/ModelContextProtocol/Private/ModelContextProtocolEngineSubsystem.cpp
deleted file mode 100644
diff --git a/Engine/Plugins/Experimental/ModelContextProtocol/Source/ModelContextProtocol/Private/ModelContextProtocolMetaData.cpp b/Engine/Plugins/Experimental/ModelContextProtocol/Source/ModelContextProtocol/Private/ModelContextProtocolMetaData.cpp
deleted file mode 100644
diff --git a/Engine/Plugins/Experimental/ModelContextProtocol/Source/ModelContextProtocol/Private/ModelContextProtocolModule.cpp b/Engine/Plugins/Experimental/ModelContextProtocol/Source/ModelContextProtocol/Private/ModelContextProtocolModule.cpp
--- a/Engine/Plugins/Experimental/ModelContextProtocol/Source/ModelContextProtocol/Private/ModelContextProtocolModule.cpp
+++ b/Engine/Plugins/Experimental/ModelContextProtocol/Source/ModelContextProtocol/Private/ModelContextProtocolModule.cpp
@@ -4,14 +4,10 @@
 #include "IModelContextProtocolResourceProvider.h"
 #include "IModelContextProtocolTool.h"
 #include "ModelContextProtocol.h"
-#include "ModelContextProtocolToolLibrary.h"
+#include "ModelContextProtocolServer.h"
 
 #include "HAL/IConsoleManager.h"
 
-#if WITH_EDITOR
-#include "KismetCompilerModule.h"
-#endif // WITH_EDITOR
-
 namespace UE::ModelContextProtocol::Private
 {
 	static constexpr const TCHAR* ModuleName = TEXT("ModelContextProtocol");
@@ -29,6 +25,45 @@ namespace UE::ModelContextProtocol::Private
 			IModelContextProtocolModule::GetChecked().RefreshTools();
 		}),
 		ECVF_Cheat);
+
+	FAutoConsoleCommand CommandStartServer = FAutoConsoleCommand(
+		TEXT("ModelContextProtocol.StartServer"),
+		TEXT("Explicitly starts the MCP HTTP server. Optional: ModelContextProtocol.StartServer <port>"),
+		FConsoleCommandWithArgsDelegate::CreateLambda([](const TArray<FString>& Args)
+		{
+			if (IModelContextProtocolModule* Module = IModelContextProtocolModule::Get())
+			{
+				if (Args.Num() > 0)
+				{
+					const int32 ParsedPort = FCString::Atoi(*Args[0]);
+					if (ParsedPort >= 1 && ParsedPort <= 65535)
+					{
+						Module->StartServer(static_cast<uint32>(ParsedPort));
+					}
+					else
+					{
+						UE_LOG(LogModelContextProtocol, Warning, TEXT("Invalid port %d. Must be 1-65535."), ParsedPort);
+					}
+				}
+				else
+				{
+					Module->StartServer();
+				}
+			}
+		}),
+		ECVF_Cheat);
+
+	FAutoConsoleCommand CommandStopServer = FAutoConsoleCommand(
+		TEXT("ModelContextProtocol.StopServer"),
+		TEXT("Explicitly stops the MCP HTTP server"),
+		FConsoleCommandWithArgsDelegate::CreateLambda([](const TArray<FString>& Args)
+		{
+			if (IModelContextProtocolModule* Module = IModelContextProtocolModule::Get())
+			{
+				Module->StopServer();
+			}
+		}),
+		ECVF_Cheat);
 }
 
 IModelContextProtocolModule* IModelContextProtocolModule::Get()
@@ -71,12 +106,29 @@ TSharedPtr<IModelContextProtocolTool> FModelContextProtocolModule::FindTool(cons
 
 bool FModelContextProtocolModule::AddTool(const TSharedRef<IModelContextProtocolTool>& Tool)
 {
-	if (FindTool(Tool->GetName()) != nullptr)
+	const FString ToolName = Tool->GetName();
+
+	if (FindTool(ToolName) != nullptr)
+	{
+		UE_LOG(LogModelContextProtocol, Warning, TEXT("A ModelContextProtocol tool named \"%s\" is already registered. All tool names must be globally unique"), *ToolName);
+		return false;
+	}
+
+	switch (UE::ModelContextProtocol::ValidateToolName(ToolName))
 	{
-		UE_LOG(LogModelContextProtocol, Warning, TEXT("A ModelContextProtocol tool named \"%s\" is already registered. All tool names must be globally unique"), *Tool->GetName());
+	case UE::ModelContextProtocol::EToolNameValidation::Empty:
+		UE_LOG(LogModelContextProtocol, Warning, TEXT("Rejecting tool with empty name; MCP spec requires 1-128 characters"));
 		return false;
+	case UE::ModelContextProtocol::EToolNameValidation::ExceedsMaxLength:
+		UE_LOG(LogModelContextProtocol, Warning, TEXT("Tool name \"%s\" exceeds MCP spec limit of 128 characters (%d)"), *ToolName, ToolName.Len());
+		break;
+	case UE::ModelContextProtocol::EToolNameValidation::InvalidCharacters:
+		UE_LOG(LogModelContextProtocol, Log, TEXT("Tool name \"%s\" contains characters outside MCP spec (allowed: A-Z, a-z, 0-9, _, -, .)"), *ToolName);
+		break;
+	default:
+		break;
 	}
-	
+
 	Tools.Add(Tool);
 
 	return true;
@@ -108,18 +160,41 @@ bool FModelContextProtocolModule::RemoveResourceProvider(const TSharedRef<IModel
 	return ResourceProviders.RemoveSingleSwap(ResourceProvider) > 0;
 }
 
+FModelContextProtocolServer* FModelContextProtocolModule::GetServer()
+{
+	return Server.Get();
+}
+
+void FModelContextProtocolModule::StartServer(uint32 Port, const FString& UrlPath)
+{
+	if (!Server)
+	{
+		Server = MakeUnique<FModelContextProtocolServer>();
+		OnRefreshToolsDelegateHandle = OnRefreshToolsDelegate.AddRaw(Server.Get(), &FModelContextProtocolServer::BroadcastToolsListChanged);
+	}
+	Server->StartServer(Port, UrlPath);
+}
+
+void FModelContextProtocolModule::StopServer()
+{
+	if (Server)
+	{
+		OnRefreshToolsDelegate.Remove(OnRefreshToolsDelegateHandle);
+		OnRefreshToolsDelegateHandle.Reset();
+		Server->StopServer();
+		Server.Reset();
+	}
+}
+
 void FModelContextProtocolModule::StartupModule()
 {
 	check(Tools.IsEmpty());
 	check(ResourceProviders.IsEmpty());
-#if WITH_EDITOR
-	IKismetCompilerInterface& KismetCompilerModule = FModuleManager::LoadModuleChecked<IKismetCompilerInterface>("KismetCompiler");
-	KismetCompilerModule.OverrideBPTypeForClass(UModelContextProtocolToolLibrary::StaticClass(), UModelContextProtocolToolLibraryBlueprint::StaticClass());
-#endif // WITH_EDITOR
 }
 
 void FModelContextProtocolModule::ShutdownModule()
 {
+	StopServer();
 	Tools.Reset();
 	ResourceProviders.Reset();
 }
diff --git a/Engine/Plugins/Experimental/ModelContextProtocol/Source/ModelContextProtocol/Private/ModelContextProtocolModule.h b/Engine/Plugins/Experimental/ModelContextProtocol/Source/ModelContextProtocol/Private/ModelContextProtocolModule.h
--- a/Engine/Plugins/Experimental/ModelContextProtocol/Source/ModelContextProtocol/Private/ModelContextProtocolModule.h
+++ b/Engine/Plugins/Experimental/ModelContextProtocol/Source/ModelContextProtocol/Private/ModelContextProtocolModule.h
@@ -6,6 +6,7 @@
 
 #include "UObject/GCObject.h"
 
+class FModelContextProtocolServer;
 struct IModelContextProtocolResourceProvider;
 struct IModelContextProtocolTool;
 
@@ -19,12 +20,16 @@ public:
 	virtual bool RemoveTool(const TSharedRef<IModelContextProtocolTool>& Tool) override;
 	virtual FOnRefreshTools& OnRefreshTools() override { return OnRefreshToolsDelegate; }
 	virtual void RefreshTools() override;
-	
+
+	[[nodiscard]] virtual FModelContextProtocolServer* GetServer() override;
+	virtual void StartServer(uint32 Port = UE::ModelContextProtocol::DefaultServerPort, const FString& UrlPath = UE::ModelContextProtocol::DefaultServerUrlPath) override;
+	virtual void StopServer() override;
+
 	virtual const TArray<TSharedRef<IModelContextProtocolResourceProvider>>& GetResourceProviders() const override;
 	virtual void AddResourceProvider(const TSharedRef<IModelContextProtocolResourceProvider>& ResourceProvider) override;
 	virtual bool RemoveResourceProvider(const TSharedRef<IModelContextProtocolResourceProvider>& ResourceProvider) override;
 	//~ End IModelContextProtocolModule implementation
-	
+
 	//~ Begin IModuleInterface implementation
 	virtual void StartupModule() override;
 	virtual void ShutdownModule() override;
@@ -39,4 +44,6 @@ protected:
 	FOnRefreshTools OnRefreshToolsDelegate;
 	TArray<TSharedRef<IModelContextProtocolTool>> Tools;
 	TArray<TSharedRef<IModelContextProtocolResourceProvider>> ResourceProviders;
+	TUniquePtr<FModelContextProtocolServer> Server;
+	FDelegateHandle OnRefreshToolsDelegateHandle;
 };
diff --git a/Engine/Plugins/Experimental/ModelContextProtocol/Source/ModelContextProtocol/Private/ModelContextProtocolServer.cpp b/Engine/Plugins/Experimental/ModelContextProtocol/Source/ModelContextProtocol/Private/ModelContextProtocolServer.cpp
new file mode 100644
--- /dev/null
+++ b/Engine/Plugins/Experimental/ModelContextProtocol/Source/ModelContextProtocol/Private/ModelContextProtocolServer.cpp
@@ -0,0 +1,1060 @@
+// Copyright Epic Games, Inc. All Rights Reserved.
+
+#include "ModelContextProtocolServer.h"
+#include "IModelContextProtocolModule.h"
+#include "IModelContextProtocolResourceProvider.h"
+#include "IModelContextProtocolTool.h"
+#include "ModelContextProtocol.h"
+
+#include "Dom/JsonObject.h"
+#include "Misc/Base64.h"
+#include "HttpServerModule.h"
+#include "IHttpRouter.h"
+#include "JsonDomBuilder.h"
+#include "Policies/CondensedJsonPrintPolicy.h"
+#include "Serialization/JsonReader.h"
+#include "Serialization/JsonSerializer.h"
+
+
+namespace UE::ModelContextProtocol::Private
+{
+	static constexpr const TCHAR* McpSessionIdHeader = TEXT("Mcp-Session-Id");
+	static constexpr const TCHAR* McpProtocolVersionHeader = TEXT("Mcp-Protocol-Version");
+
+	bool GetProtocolVersionFromRequestHeaders(const FHttpServerRequest& Request, FString& OutProtocolVersion)
+	{
+		if (const TArray<FString>* Headers = Request.Headers.Find(McpProtocolVersionHeader))
+		{
+			if (!Headers->IsEmpty())
+			{
+				OutProtocolVersion = (*Headers)[0];
+				return true;
+			}
+		}
+		return false;
+	}
+
+	void CompleteWithResponseCode(const FHttpResultCallback& OnComplete, EHttpServerResponseCodes ResponseCode);
+
+	/**
+	 * Validates the Origin header to prevent DNS rebinding attacks per MCP spec.
+	 * Allows requests with no Origin (non-browser clients), or Origin matching localhost variants.
+	 * Returns true if the request is allowed, false if it was rejected (and OnComplete was called).
+	 */
+	bool ValidateOriginHeader(const FHttpServerRequest& Request, const FHttpResultCallback& OnComplete)
+	{
+		const TArray<FString>* OriginHeaders = Request.Headers.Find(TEXT("Origin"));
+		if (!OriginHeaders || OriginHeaders->IsEmpty())
+		{
+			// No Origin header — non-browser client, allow
+			return true;
+		}
+
+		const FString& Origin = (*OriginHeaders)[0];
+
+		// Extract the host portion from the Origin (scheme://host[:port])
+		FString OriginHost;
+		FString SchemeAndRest;
+		if (Origin.Split(TEXT("://"), &SchemeAndRest, &OriginHost))
+		{
+			// Strip port if present, handling bracketed IPv6 addresses (e.g., [::1]:8080)
+			FString Host;
+			if (OriginHost.StartsWith(TEXT("[")))
+			{
+				int32 CloseBracket;
+				if (OriginHost.FindChar(TEXT(']'), CloseBracket))
+				{
+					Host = OriginHost.Left(CloseBracket + 1);
+				}
+				else
+				{
+					Host = OriginHost;
+				}
+			}
+			else
+			{
+				FString Port;
+				if (!OriginHost.Split(TEXT(":"), &Host, &Port))
+				{
+					Host = OriginHost;
+				}
+			}
+
+			if (Host == TEXT("localhost") || Host == TEXT("127.0.0.1") || Host == TEXT("[::1]"))
+			{
+				return true;
+			}
+		}
+
+		UE_LOG(LogModelContextProtocol, Warning, TEXT("Rejected request with disallowed Origin: %s"), *Origin);
+		CompleteWithResponseCode(OnComplete, EHttpServerResponseCodes::Forbidden);
+		return false;
+	}
+
+	enum class EJsonRpcErrorCode : int32
+	{
+		/**
+		 * Invalid JSON was received by the server.
+		 * An error occurred on the server while parsing the JSON text.
+		 */
+		ParseError = -32700,
+
+		/** The JSON sent is not a valid Request object. */
+		InvalidRequest = -32600,
+
+		/** The method does not exist / is not available. */
+		MethodNotFound = -32601,
+
+		/** Invalid method parameter(s). */
+		InvalidParams = -32602,
+
+		/** Internal JSON-RPC error. */
+		InternalError = -32603,
+
+		/** Resource not found (MCP-specific). */
+		ResourceNotFound = -32002
+	};
+
+	TSharedPtr<FJsonObject> GetJsonObjectFromRequestBody(const TArray<uint8>& InRequestBody)
+	{
+		const FUtf8StringView IncomingRequestBody(reinterpret_cast<const UTF8CHAR*>(InRequestBody.GetData()), InRequestBody.Num());
+		TSharedPtr<FJsonObject> BodyObject = MakeShared<FJsonObject>();
+		TSharedRef<TJsonReader<UTF8CHAR>> JsonReader = TJsonReaderFactory<UTF8CHAR>::CreateFromView(IncomingRequestBody);
+
+		if (FJsonSerializer::Deserialize(JsonReader, BodyObject) && BodyObject.IsValid())
+		{
+			return BodyObject;
+		}
+
+		return nullptr;
+	}
+
+	TUniquePtr<FHttpServerResponse> CreateJsonRpcErrorResponse(const TSharedPtr<FJsonValue>& RequestId, int32 ErrorCode, const FString& ErrorMessage, EHttpServerResponseCodes ResponseCode = EHttpServerResponseCodes::BadRequest)
+	{
+		FJsonDomBuilder::FObject ErrorObject;
+		ErrorObject.Set(TEXT("code"), ErrorCode);
+		ErrorObject.Set(TEXT("message"), ErrorMessage);
+
+		FJsonDomBuilder::FObject ResponseObject;
+		ResponseObject.Set(TEXT("jsonrpc"), UE::ModelContextProtocol::JsonRpcVersion);
+		ResponseObject.Set(TEXT("id"), RequestId.IsValid() ? RequestId : MakeShared<FJsonValueNull>());
+		ResponseObject.Set(TEXT("error"), ErrorObject);
+
+		FUtf8String ResponseStr;
+		TSharedRef<TJsonWriter<UTF8CHAR>> JsonWriter = TJsonWriterFactory<UTF8CHAR>::Create(&ResponseStr);
+		FJsonSerializer::Serialize(ResponseObject.AsJsonObject(), JsonWriter);
+
+		TUniquePtr<FHttpServerResponse> Response = FHttpServerResponse::Create(MoveTemp(ResponseStr), TEXT("application/json"));
+		Response->Code = ResponseCode;
+		return Response;
+	}
+
+
+	TUniquePtr<FHttpServerResponse> CreateJsonRpcResultResponse(const TSharedPtr<FJsonValue>& RequestId, const TSharedPtr<FJsonValue>& Result, const FString* SessionId = nullptr)
+	{
+		if (!ensureMsgf(RequestId.IsValid(), TEXT("JSON-RPC responses *MUST* have a valid request id. Requests without an id must be treated as notifications and should not be replied to.")))
+		{
+			return CreateJsonRpcErrorResponse(RequestId, static_cast<int32>(EJsonRpcErrorCode::InternalError), TEXT("Replying to request with unknown RequestId"));
+		}
+
+		FJsonDomBuilder::FObject ResponseObject;
+		ResponseObject.Set(TEXT("jsonrpc"), JsonRpcVersion);
+		ResponseObject.Set(TEXT("id"), RequestId);
+		ResponseObject.Set(TEXT("result"), Result);
+
+		FUtf8String ResponseStr;
+		TSharedRef<TJsonWriter<UTF8CHAR>> JsonWriter = TJsonWriterFactory<UTF8CHAR>::Create(&ResponseStr);
+		FJsonSerializer::Serialize(ResponseObject.AsJsonObject(), JsonWriter);
+
+		TUniquePtr<FHttpServerResponse> ServerResponse = FHttpServerResponse::Create(MoveTemp(ResponseStr), TEXT("application/json"));
+		if (SessionId)
+		{
+			ServerResponse->Headers.Add(McpSessionIdHeader, {*SessionId});
+		}
+		return ServerResponse;
+	}
+
+	void CompleteWithError(const FHttpResultCallback& OnComplete, const TSharedPtr<FJsonValue>& RequestId, UE::ModelContextProtocol::Private::EJsonRpcErrorCode ErrorCode, const FString& ErrorMessage, EHttpServerResponseCodes ResponseCode = EHttpServerResponseCodes::BadRequest)
+	{
+		UE_LOG(LogModelContextProtocol, Error, TEXT("%s"), *ErrorMessage);
+		OnComplete(CreateJsonRpcErrorResponse(RequestId, static_cast<int32>(ErrorCode), ErrorMessage, ResponseCode));
+	}
+
+	void CompleteWithResult(const FHttpResultCallback& OnComplete, const TSharedPtr<FJsonValue>& RequestId, const TSharedPtr<FJsonValue>& Result, const FString* SessionId = nullptr)
+	{
+		OnComplete(CreateJsonRpcResultResponse(RequestId, Result, SessionId));
+	}
+
+	void CompleteWithResult(const FHttpResultCallback& OnComplete, const TSharedPtr<FJsonValue>& RequestId, const FJsonDomBuilder::FObject& Result, const FString* SessionId = nullptr)
+	{
+		OnComplete(CreateJsonRpcResultResponse(RequestId, Result.AsJsonValue(), SessionId));
+	}
+
+	template <typename ResultType>
+	void CompleteWithResult(const FHttpResultCallback& OnComplete, const TSharedPtr<FJsonValue>& RequestId, const ResultType& Result, const FString* SessionId = nullptr)
+	{
+		TSharedPtr<FJsonObject> JsonResult = FJsonObjectConverter::UStructToJsonObject<ResultType>(Result);
+		if (!ensure(JsonResult.IsValid()))
+		{
+			CompleteWithError(
+				OnComplete, RequestId, EJsonRpcErrorCode::InternalError, FString::Printf(TEXT("Error forming response")));
+			return;
+		}
+
+		TSharedPtr<FJsonValue> JsonResultValue = MakeShared<FJsonValueObject>(JsonResult);
+		CompleteWithResult(OnComplete, RequestId, JsonResultValue, SessionId);
+	}
+
+	void CompleteWithResponseCode(const FHttpResultCallback& OnComplete, EHttpServerResponseCodes ResponseCode)
+	{
+		TUniquePtr<FHttpServerResponse> Response = MakeUnique<FHttpServerResponse>();
+		Response->Code = ResponseCode;
+		OnComplete(MoveTemp(Response));
+	}
+
+	template <typename ParamsType>
+	bool ParseParamsOrError(const TCHAR* MethodName, const TSharedPtr<FJsonValue>& RequestId, const FHttpResultCallback& OnComplete, const TSharedPtr<FJsonObject>& JsonParams, ParamsType& OutParams)
+	{
+		// Parse params
+		if (!JsonParams.IsValid())
+		{
+			CompleteWithError(
+				OnComplete, RequestId, EJsonRpcErrorCode::InvalidParams, FString::Printf(TEXT("Expected params for call to \"%s\""), MethodName));
+			return false;
+		}
+		FText OutFailReason;
+		if (!FJsonObjectConverter::JsonObjectToUStruct<ParamsType>(JsonParams.ToSharedRef(), &OutParams, /*CheckFlags*/CPF_None, /*SkipFlags*/CPF_None, /*bStrictMode*/false, &OutFailReason))
+		{
+			CompleteWithError(
+				OnComplete, RequestId, EJsonRpcErrorCode::InvalidParams, FString::Printf(TEXT("Invalid %s params: %s"), MethodName, *OutFailReason.ToString()));
+			return false;
+		}
+
+		return true;
+	}
+
+	bool GetSessionIdFromRequestHeaders(const FHttpServerRequest& Request, FString& OutSessionId)
+	{
+		if (const TArray<FString>* McpSessionIdHeaders = Request.Headers.Find(McpSessionIdHeader))
+		{
+			if (!McpSessionIdHeaders->IsEmpty())
+			{
+				OutSessionId = (*McpSessionIdHeaders)[0];
+				return true;
+			}
+		}
+		return false;
+	}
+
+	FUtf8String FormatSSEMessage(const FUtf8String& Message)
+	{
+		return UTF8TEXT("event: message\r\ndata: ") + Message + UTF8TEXT("\r\n\r\n");
+	}
+
+	TUniquePtr<FHttpServerResponse> CreateJsonRpcProgressResponse(const TSharedRef<FJsonValue>& ProgressToken, int32 Progress)
+	{
+		FJsonDomBuilder::FObject ParamsObject;
+		ParamsObject.Set(TEXT("progressToken"), ProgressToken);
+		ParamsObject.Set(TEXT("progress"), Progress);
+
+		FJsonDomBuilder::FObject ResponseObject;
+		ResponseObject.Set(TEXT("jsonrpc"), UE::ModelContextProtocol::JsonRpcVersion);
+		ResponseObject.Set(TEXT("method"), TEXT("notifications/progress"));
+		ResponseObject.Set(TEXT("params"), ParamsObject);
+
+		FUtf8String ResponseStr;
+		TSharedRef<TJsonWriter<UTF8CHAR, TCondensedJsonPrintPolicy<UTF8CHAR>>> JsonWriter = TJsonWriterFactory<UTF8CHAR, TCondensedJsonPrintPolicy<UTF8CHAR>>::Create(&ResponseStr);
+		FJsonSerializer::Serialize(ResponseObject.AsJsonObject(), JsonWriter);
+
+		TUniquePtr<FHttpServerResponse> Response = FHttpServerResponse::Create(UE::ModelContextProtocol::Private::FormatSSEMessage(ResponseStr), UE::ModelContextProtocol::ContentTypeEventStream);
+		Response->Code = EHttpServerResponseCodes::Ok;
+		EnumAddFlags(Response->Flags, EHttpServerResponseFlags::MultipleWriteStream | EHttpServerResponseFlags::HasAdditionalWrites | EHttpServerResponseFlags::SkipHeaderWrite);
+
+		return Response;
+	}
+
+	void SendProgressUpdate(const FHttpResultCallback& ResponseCallback, const TSharedRef<FJsonValue>& ProgressToken, int32 Progress)
+	{
+		ResponseCallback(CreateJsonRpcProgressResponse(ProgressToken, Progress));
+	}
+
+	TUniquePtr<FHttpServerResponse> CreateJsonRpcNotificationResponse(const FString& Method)
+	{
+		FJsonDomBuilder::FObject NotificationObject;
+		NotificationObject.Set(TEXT("jsonrpc"), UE::ModelContextProtocol::JsonRpcVersion);
+		NotificationObject.Set(TEXT("method"), Method);
+		NotificationObject.Set(TEXT("params"), FJsonDomBuilder::FObject());
+
+		FUtf8String ResponseStr;
+		TSharedRef<TJsonWriter<UTF8CHAR, TCondensedJsonPrintPolicy<UTF8CHAR>>> JsonWriter = TJsonWriterFactory<UTF8CHAR, TCondensedJsonPrintPolicy<UTF8CHAR>>::Create(&ResponseStr);
+		FJsonSerializer::Serialize(NotificationObject.AsJsonObject(), JsonWriter);
+
+		TUniquePtr<FHttpServerResponse> Response = FHttpServerResponse::Create(UE::ModelContextProtocol::Private::FormatSSEMessage(ResponseStr), UE::ModelContextProtocol::ContentTypeEventStream);
+		Response->Code = EHttpServerResponseCodes::Ok;
+		EnumAddFlags(Response->Flags, EHttpServerResponseFlags::MultipleWriteStream | EHttpServerResponseFlags::HasAdditionalWrites | EHttpServerResponseFlags::SkipHeaderWrite);
+
+		return Response;
+	}
+
+	struct FPaginationResult
+	{
+		TArray<TSharedPtr<FJsonValue>> Items;
+		TOptional<FString> NextCursor;
+	};
+
+	bool ApplyPagination(const TArray<TSharedPtr<FJsonValue>>& AllItems, const TSharedPtr<FJsonObject>& Params, FPaginationResult& OutResult, FString& OutError)
+	{
+		const int32 PageSize = UE::ModelContextProtocol::PaginationPageSize;
+
+		if (PageSize <= 0)
+		{
+			OutResult.Items = AllItems;
+			return true;
+		}
+
+		int32 Offset = 0;
+
+		if (Params.IsValid())
+		{
+			FString Cursor;
+			if (Params->TryGetStringField(TEXT("cursor"), Cursor) && !Cursor.IsEmpty())
+			{
+				FString DecodedCursor;
+				auto IsAllDigits = [](const FString& String)
+				{
+					for (TCHAR Character : String)
+					{
+						if (!FChar::IsDigit(Character))
+						{
+							return false;
+						}
+					}
+					return true;
+				};
+				if (!FBase64::Decode(Cursor, DecodedCursor) || DecodedCursor.IsEmpty() || !IsAllDigits(DecodedCursor))
+				{
+					OutError = TEXT("Invalid pagination cursor");
+					return false;
+				}
+
+				const int64 DecodedOffset = FCString::Atoi64(*DecodedCursor);
+				if (DecodedOffset < 0 || DecodedOffset > MAX_int32)
+				{
+					OutError = TEXT("Invalid pagination cursor");
+					return false;
+				}
+				Offset = static_cast<int32>(DecodedOffset);
+			}
+		}
+
+		if (Offset > AllItems.Num())
+		{
+			Offset = AllItems.Num();
+		}
+
+		const int32 End = static_cast<int32>(FMath::Min(static_cast<int64>(Offset) + PageSize, static_cast<int64>(AllItems.Num())));
+
+		OutResult.Items.Reserve(End - Offset);
+		for (int32 Index = Offset; Index < End; ++Index)
+		{
+			OutResult.Items.Add(AllItems[Index]);
+		}
+
+		if (End < AllItems.Num())
+		{
+			OutResult.NextCursor = FBase64::Encode(FString::FromInt(End));
+		}
+
+		return true;
+	}
+}
+
+FModelContextProtocolServer::~FModelContextProtocolServer()
+{
+	AliveGuard.Reset();
+	StopServer();
+}
+
+TSharedPtr<FModelContextProtocolSession> FModelContextProtocolServer::FindSession(const FString& SessionId) const
+{
+	const TSharedPtr<FModelContextProtocolSession>* SessionPtr = Sessions.FindByPredicate([&SessionId](const TSharedPtr<FModelContextProtocolSession>& Session)
+		{
+			return Session->ID == SessionId;
+		});
+	return SessionPtr ? *SessionPtr : nullptr;
+}
+
+void FModelContextProtocolServer::StartServer(uint32 Port, const FString& UrlPath)
+{
+	if (HttpRouter.IsValid())
+	{
+		UE_LOG(LogModelContextProtocol, Warning, TEXT("MCP server is already running on port %u, stopping first."), ActiveServerPort);
+		StopServer();
+	}
+
+	ActiveServerPort = Port;
+	HttpRouter = FHttpServerModule::Get().GetHttpRouter(ActiveServerPort);
+	if (ensure(HttpRouter.IsValid()))
+	{
+		const FHttpPath MCPServerPath = UrlPath;
+
+		MainMcpRoute = HttpRouter->BindRoute(MCPServerPath, EHttpServerRequestVerbs::VERB_POST,
+			FHttpRequestHandler::CreateRaw(this, &FModelContextProtocolServer::ProcessPostRequest));
+
+		SseMcpRoute = HttpRouter->BindRoute(MCPServerPath, EHttpServerRequestVerbs::VERB_GET,
+			FHttpRequestHandler::CreateRaw(this, &FModelContextProtocolServer::ProcessGetRequest));
+
+		DeleteMcpRoute = HttpRouter->BindRoute(MCPServerPath, EHttpServerRequestVerbs::VERB_DELETE,
+			FHttpRequestHandler::CreateRaw(this, &FModelContextProtocolServer::ProcessDeleteRequest));
+	}
+
+	FHttpServerModule::Get().StartAllListeners();
+
+	TickerHandle = FTSTicker::GetCoreTicker().AddTicker(TEXT("ModelContextProtocolServer"), 0.0f, [this](float DeltaTime)
+	{
+		Tick(DeltaTime);
+		return true;
+	});
+}
+
+void FModelContextProtocolServer::StopServer()
+{
+	if (TickerHandle.IsValid())
+	{
+		FTSTicker::RemoveTicker(TickerHandle);
+		TickerHandle.Reset();
+	}
+
+	if (HttpRouter.IsValid())
+	{
+		if (MainMcpRoute.IsValid())
+		{
+			HttpRouter->UnbindRoute(MainMcpRoute);
+		}
+
+		if (SseMcpRoute.IsValid())
+		{
+			HttpRouter->UnbindRoute(SseMcpRoute);
+		}
+
+		if (DeleteMcpRoute.IsValid())
+		{
+			HttpRouter->UnbindRoute(DeleteMcpRoute);
+		}
+	}
+
+	MainMcpRoute.Reset();
+	SseMcpRoute.Reset();
+	DeleteMcpRoute.Reset();
+	HttpRouter.Reset();
+	ActiveServerPort = 0;
+}
+
+bool FModelContextProtocolServer::IsServerRunning() const
+{
+	return HttpRouter.IsValid();
+}
+
+uint32 FModelContextProtocolServer::GetServerPort() const
+{
+	return ActiveServerPort;
+}
+
+void FModelContextProtocolServer::BroadcastToolsListChanged()
+{
+	for (const TSharedPtr<FModelContextProtocolSession>& Session : Sessions)
+	{
+		if (!Session.IsValid() || Session->Status != EModelContextProtocolSessionStatus::Initialized)
+		{
+			continue;
+		}
+
+		for (auto& [RequestId, Context] : Session->ActiveRequests)
+		{
+			if (Context.EventStreamWrite.IsSet())
+			{
+				Context.EventStreamWrite(UE::ModelContextProtocol::Private::CreateJsonRpcNotificationResponse(TEXT("notifications/tools/list_changed")));
+				// Only need to send once per session — all active SSE streams share the same client.
+				break;
+			}
+		}
+	}
+
+	UE_LOG(LogModelContextProtocol, Log, TEXT("Broadcast notifications/tools/list_changed to active sessions"));
+}
+
+bool FModelContextProtocolServer::ProcessPostRequest(const FHttpServerRequest& Request, const FHttpResultCallback& OnComplete)
+{
+	if (!UE::ModelContextProtocol::Private::ValidateOriginHeader(Request, OnComplete))
+	{
+		return true;
+	}
+
+	TSharedPtr<FJsonObject> RequestBodyObject = UE::ModelContextProtocol::Private::GetJsonObjectFromRequestBody(Request.Body);
+	if (!RequestBodyObject.IsValid())
+	{
+		UE::ModelContextProtocol::Private::CompleteWithError(
+			OnComplete, /*RequestId*/{}, UE::ModelContextProtocol::Private::EJsonRpcErrorCode::ParseError, TEXT("Invalid JSON body!"));
+		return true;
+	}
+
+	// jsonrpc
+	FString RequestJsonRpcVersion;
+	if (!RequestBodyObject->TryGetStringField(TEXT("jsonrpc"), RequestJsonRpcVersion) || RequestJsonRpcVersion != UE::ModelContextProtocol::JsonRpcVersion)
+	{
+		UE::ModelContextProtocol::Private::CompleteWithError(
+			OnComplete, /*RequestId*/{}, UE::ModelContextProtocol::Private::EJsonRpcErrorCode::InvalidRequest,
+			FString::Printf(TEXT("Unexpected jsonrpc version %s. Must be %s"), *RequestJsonRpcVersion, UE::ModelContextProtocol::JsonRpcVersion));
+		return true;
+	}
+
+	// id
+	TSharedPtr<FJsonValue> RequestId = RequestBodyObject->TryGetField(TEXT("id"));
+
+	// method
+	FString Method;
+	if (!RequestBodyObject->TryGetStringField(TEXT("method"), Method))
+	{
+		UE::ModelContextProtocol::Private::CompleteWithError(
+			OnComplete, /*RequestId*/{}, UE::ModelContextProtocol::Private::EJsonRpcErrorCode::InvalidRequest, TEXT("Missing required field \"method\""));
+		return true;
+	}
+
+	// params
+	const TSharedPtr<FJsonObject>* ParamsObject = nullptr;
+	RequestBodyObject->TryGetObjectField(TEXT("params"), ParamsObject);
+
+	return ProcessJsonRpcCall(Request, RequestId, Method, ParamsObject ? *ParamsObject : nullptr, OnComplete);
+}
+
+bool FModelContextProtocolServer::ProcessJsonRpcCall(const FHttpServerRequest& Request, const TSharedPtr<FJsonValue>& RequestId, const FString& Method, const TSharedPtr<FJsonObject>& Params, const FHttpResultCallback& OnComplete)
+{
+	TSharedPtr<FJsonValue> Result;
+
+	// ping
+	if (Method == TEXT("ping"))
+	{
+		return ProcessPingJsonRpcCall(Request, RequestId, OnComplete);
+	}
+	// initialize
+	if (Method == TEXT("initialize"))
+	{
+		return ProcessInitializeJsonRpcCall(Request, RequestId, Params, OnComplete);
+	}
+
+	// Per MCP spec (2025-11-25): validate Mcp-Protocol-Version header on all post-initialize requests.
+	// If present, verify it matches the session's negotiated version.
+	FString HeaderProtocolVersion;
+	if (UE::ModelContextProtocol::Private::GetProtocolVersionFromRequestHeaders(Request, HeaderProtocolVersion))
+	{
+		FString SessionId;
+		if (UE::ModelContextProtocol::Private::GetSessionIdFromRequestHeaders(Request, SessionId))
+		{
+			if (TSharedPtr<FModelContextProtocolSession> Session = FindSession(SessionId))
+			{
+				if (HeaderProtocolVersion != Session->NegotiatedProtocolVersion)
+				{
+					UE::ModelContextProtocol::Private::CompleteWithError(
+						OnComplete, RequestId, UE::ModelContextProtocol::Private::EJsonRpcErrorCode::InvalidRequest,
+						FString::Printf(TEXT("Mcp-Protocol-Version header '%s' does not match negotiated version '%s'"), *HeaderProtocolVersion, *Session->NegotiatedProtocolVersion));
+					return true;
+				}
+			}
+		}
+	}
+
+	// notifications/initialized
+	if (Method == TEXT("notifications/initialized"))
+	{
+		return ProcessInitializedNotificationJsonRpcCall(Request, RequestId, Params, OnComplete);
+	}
+	// notifications/cancelled
+	else if (Method == TEXT("notifications/cancelled"))
+	{
+		return ProcessNotificationCancelledJsonRpcCall(Request, RequestId, Params, OnComplete);
+	}
+	// tools/list
+	else if (Method == TEXT("tools/list"))
+	{
+		return ProcessListToolsJsonRpcCall(Request, RequestId, Params, OnComplete);
+	}
+	// tools/call
+	else if (Method == TEXT("tools/call"))
+	{
+		return ProcessToolCallJsonRpcCall(Request, RequestId, Params, OnComplete);
+	}
+	// resources/list
+	else if (Method == TEXT("resources/list"))
+	{
+		return ProcessListResourcesJsonRpcCall(Request, RequestId, Params, OnComplete);
+	}
+	// resources/read
+	else if (Method == TEXT("resources/read"))
+	{
+		return ProcessReadResourceJsonRpcCall(Request, RequestId, Params, OnComplete);
+	}
+
+	UE::ModelContextProtocol::Private::CompleteWithError(
+		OnComplete, RequestId, UE::ModelContextProtocol::Private::EJsonRpcErrorCode::MethodNotFound, FString::Printf(TEXT("Call to unknown method \"%s\""), *Method));
+	return true;
+}
+
+bool FModelContextProtocolServer::ProcessPingJsonRpcCall(const FHttpServerRequest& Request, const TSharedPtr<FJsonValue>& RequestId, const FHttpResultCallback& OnComplete)
+{
+	// Reply with an empty result per spec for ping.
+	UE_LOG(LogModelContextProtocol, Log, TEXT("Replying to Ping"));
+	FModelContextProtocolPingResult Result;
+	UE::ModelContextProtocol::Private::CompleteWithResult(OnComplete, RequestId, Result);
+	return true;
+}
+
+bool FModelContextProtocolServer::ProcessInitializeJsonRpcCall(const FHttpServerRequest& Request, const TSharedPtr<FJsonValue>& RequestId, const TSharedPtr<FJsonObject>& Params, const FHttpResultCallback& OnComplete)
+{
+	FModelContextProtocolInitializeParams InitializeParams;
+	if (!UE::ModelContextProtocol::Private::ParseParamsOrError(TEXT("initialize"), RequestId, OnComplete, Params, InitializeParams))
+	{
+		// Request was completed (with error)
+		return true;
+	}
+
+	// Negotiate protocol version per MCP spec:
+	// If the server supports the client's version, respond with that version.
+	// Otherwise, respond with the server's latest supported version.
+	const FString NegotiatedVersion = UE::ModelContextProtocol::NegotiateProtocolVersion(InitializeParams.ProtocolVersion);
+	UE_LOG(LogModelContextProtocol, Display, TEXT("Client requested protocol version '%s', negotiated '%s'"), *InitializeParams.ProtocolVersion, *NegotiatedVersion);
+
+	// Create a new session for this client
+	TSharedPtr<FModelContextProtocolSession> Session = MakeShared<FModelContextProtocolSession>();
+	Session->ID = FGuid::NewGuid().ToString(EGuidFormats::DigitsLower);
+	Session->Status = EModelContextProtocolSessionStatus::Initializing;
+	Session->NegotiatedProtocolVersion = NegotiatedVersion;
+	Session->ClientAddress = Request.PeerAddress;
+	Session->ClientCapabilities = InitializeParams.Capabilities;
+	Sessions.Add(Session);
+	UE_LOG(LogModelContextProtocol, Log, TEXT("Initializing new session: %s"), *Session->ID);
+
+	// Reply with server capabilities
+	FModelContextProtocolInitializeResult InitializeResult;
+	InitializeResult.ProtocolVersion = NegotiatedVersion;
+	FModelContextProtocolToolsCapability ToolsCapability;
+	ToolsCapability.ListChanged = true;
+	InitializeResult.Capabilities.Tools = ToolsCapability;
+	InitializeResult.Capabilities.Resources = FModelContextProtocolResourcesCapability();
+	UE::ModelContextProtocol::Private::CompleteWithResult(OnComplete, RequestId, InitializeResult, &Session->ID);
+	return true;
+}
+
+bool FModelContextProtocolServer::ProcessInitializedNotificationJsonRpcCall(const FHttpServerRequest& Request, const TSharedPtr<FJsonValue>& RequestId, const TSharedPtr<FJsonObject>& Params, const FHttpResultCallback& OnComplete)
+{
+	FString SessionId;
+	if (!UE::ModelContextProtocol::Private::GetSessionIdFromRequestHeaders(Request, SessionId))
+	{
+		UE::ModelContextProtocol::Private::CompleteWithError(OnComplete, /*RequestId*/nullptr, UE::ModelContextProtocol::Private::EJsonRpcErrorCode::InvalidRequest, TEXT("Missing Mcp-Session-Id header for notifications/initialized"), EHttpServerResponseCodes::BadRequest);
+		return true;
+	}
+
+	TSharedPtr<FModelContextProtocolSession> Session = FindSession(SessionId);
+	if (!Session.IsValid())
+	{
+		UE::ModelContextProtocol::Private::CompleteWithError(OnComplete, /*RequestId*/nullptr, UE::ModelContextProtocol::Private::EJsonRpcErrorCode::InvalidRequest, FString::Printf(TEXT("Received notifications/initialized for unknown session: %s"), *SessionId), EHttpServerResponseCodes::BadRequest);
+		return true;
+	}
+
+	Session->Status = EModelContextProtocolSessionStatus::Initialized;
+	UE_LOG(LogModelContextProtocol, Log, TEXT("Session initialized: %s"), *Session->ID);
+
+	// Accept notification
+	UE::ModelContextProtocol::Private::CompleteWithResponseCode(OnComplete, EHttpServerResponseCodes::Accepted);
+	return true;
+}
+
+bool FModelContextProtocolServer::ProcessNotificationCancelledJsonRpcCall(const FHttpServerRequest& Request, const TSharedPtr<FJsonValue>& RequestId, const TSharedPtr<FJsonObject>& Params, const FHttpResultCallback& OnComplete)
+{
+	FString SessionId;
+	if (UE::ModelContextProtocol::Private::GetSessionIdFromRequestHeaders(Request, SessionId))
+	{
+		TSharedPtr<FModelContextProtocolSession> Session = FindSession(SessionId);
+		if (!Session.IsValid())
+		{
+			UE::ModelContextProtocol::Private::CompleteWithError(OnComplete, /*RequestId*/nullptr, UE::ModelContextProtocol::Private::EJsonRpcErrorCode::InvalidRequest, FString::Printf(TEXT("Received notifications/cancelled for unknown session: %s"), *SessionId), EHttpServerResponseCodes::BadRequest);
+			return true;
+		}
+
+		if (!Params.IsValid())
+		{
+			UE::ModelContextProtocol::Private::CompleteWithError(OnComplete, /*RequestId*/nullptr, UE::ModelContextProtocol::Private::EJsonRpcErrorCode::InvalidRequest, FString::Printf(TEXT("Expected params for notifications/cancelled for session: %s"), *SessionId), EHttpServerResponseCodes::BadRequest);
+			return true;
+		}
+
+		const TSharedPtr<FJsonValue> CancelId = Params->TryGetField(TEXT("requestId"));
+
+		if (!CancelId.IsValid())
+		{
+			UE::ModelContextProtocol::Private::CompleteWithError(OnComplete, /*RequestId*/nullptr, UE::ModelContextProtocol::Private::EJsonRpcErrorCode::InvalidRequest, FString::Printf(TEXT("Expected requestId for notifications/cancelled for session: %s"), *SessionId), EHttpServerResponseCodes::BadRequest);
+			return true;
+		}
+
+		FModelContextProtocolToolRequestId ToolRequestId(CancelId.ToSharedRef());
+
+		if (const FModelContextProtocolToolContext* Context = Session->ActiveRequests.Find(ToolRequestId))
+		{
+			if (Context->Tool.IsValid())
+			{
+				Context->Tool->CancelAsync(ToolRequestId);
+			}
+
+			Session->ActiveRequests.Remove(ToolRequestId);
+		}
+		else
+		{
+			// Request may have already finished, so not necessarily an error
+		}
+	}
+
+	// Accept notification
+	UE::ModelContextProtocol::Private::CompleteWithResponseCode(OnComplete, EHttpServerResponseCodes::Accepted);
+	return true;
+}
+
+bool FModelContextProtocolServer::ProcessListToolsJsonRpcCall(const FHttpServerRequest& Request, const TSharedPtr<FJsonValue>& RequestId, const TSharedPtr<FJsonObject>& Params, const FHttpResultCallback& OnComplete)
+{
+	FJsonDomBuilder::FArray ToolsArray;
+
+	const IModelContextProtocolModule& ModelContextProtocolModule = IModelContextProtocolModule::GetChecked();
+	for (const TSharedRef<IModelContextProtocolTool>& Tool : ModelContextProtocolModule.GetTools())
+	{
+		FJsonDomBuilder::FObject ToolObject;
+		ToolObject.Set(TEXT("name"), Tool->GetName());
+		if (const FString Description = Tool->GetDescription(); !Description.IsEmpty())
+		{
+			ToolObject.Set(TEXT("description"), Description);
+		}
+		if (const TSharedPtr<FJsonObject> InputJsonSchemaObject = Tool->GetInputJsonSchema())
+		{
+			ToolObject.Set(TEXT("inputSchema"), InputJsonSchemaObject);
+		}
+		if (const TSharedPtr<FJsonObject> OutputJsonSchemaObject = Tool->GetOutputJsonSchema())
+		{
+			ToolObject.Set(TEXT("outputSchema"), OutputJsonSchemaObject);
+		}
+		ToolsArray.Add(ToolObject);
+	}
+
+	UE_LOG(LogModelContextProtocol, Log, TEXT("Listing tools (%d)"), ToolsArray.Num());
+
+	TSharedRef<FJsonValueArray> AllToolsJsonValue = ToolsArray.AsJsonValue();
+	const TArray<TSharedPtr<FJsonValue>>& AllTools = AllToolsJsonValue->AsArray();
+	UE::ModelContextProtocol::Private::FPaginationResult PaginationResult;
+	FString PaginationError;
+	if (!UE::ModelContextProtocol::Private::ApplyPagination(AllTools, Params, PaginationResult, PaginationError))
+	{
+		UE::ModelContextProtocol::Private::CompleteWithError(OnComplete, RequestId, UE::ModelContextProtocol::Private::EJsonRpcErrorCode::InvalidParams, PaginationError);
+		return true;
+	}
+
+	FJsonDomBuilder::FObject ResultObject;
+	ResultObject.Set(TEXT("tools"), MakeShared<FJsonValueArray>(MoveTemp(PaginationResult.Items)));
+	if (PaginationResult.NextCursor.IsSet())
+	{
+		ResultObject.Set(TEXT("nextCursor"), PaginationResult.NextCursor.GetValue());
+	}
+
+	UE::ModelContextProtocol::Private::CompleteWithResult(OnComplete, RequestId, ResultObject);
+	return true;
+}
+
+bool FModelContextProtocolServer::ProcessToolCallJsonRpcCall(const FHttpServerRequest& Request, const TSharedPtr<FJsonValue>& RequestId, const TSharedPtr<FJsonObject>& Params, const FHttpResultCallback& OnComplete)
+{
+	if (!Params.IsValid())
+	{
+		UE::ModelContextProtocol::Private::CompleteWithError(OnComplete, RequestId, UE::ModelContextProtocol::Private::EJsonRpcErrorCode::InvalidParams, TEXT("Expected params for call to tools/call"));
+		return true;
+	}
+
+	FString ToolName;
+	if (!Params->TryGetStringField(TEXT("name"), ToolName) || ToolName.IsEmpty())
+	{
+		UE::ModelContextProtocol::Private::CompleteWithError(OnComplete, RequestId, UE::ModelContextProtocol::Private::EJsonRpcErrorCode::InvalidParams, TEXT("Expected non-empty 'name' param for call to tools/call"));
+		return true;
+	}
+	if (!RequestId.IsValid())
+	{
+		UE::ModelContextProtocol::Private::CompleteWithError(OnComplete, RequestId, UE::ModelContextProtocol::Private::EJsonRpcErrorCode::InvalidParams, TEXT("Invalid request id"));
+		return true;
+	}
+
+	TSharedPtr<FJsonObject> ToolArguments;
+	if (Params->HasTypedField<EJson::Object>(TEXT("arguments")))
+	{
+		ToolArguments = Params->GetObjectField(TEXT("arguments"));
+	}
+
+	const IModelContextProtocolModule& ModelContextProtocolModule = IModelContextProtocolModule::GetChecked();
+	TSharedPtr<IModelContextProtocolTool> Tool = ModelContextProtocolModule.FindTool(ToolName);
+	if (!Tool.IsValid())
+	{
+		UE::ModelContextProtocol::Private::CompleteWithError(OnComplete, RequestId, UE::ModelContextProtocol::Private::EJsonRpcErrorCode::InvalidParams, FString::Printf(TEXT("Unknown tool: %s"), *ToolName));
+		return true;
+	}
+
+	UE_LOG(LogModelContextProtocol, Log, TEXT("Running tool: '%s'"), *ToolName);
+	FString PrettyJson;
+	if (UE_LOG_ACTIVE(LogModelContextProtocol, VeryVerbose) && ToolArguments.IsValid() && FJsonSerializer::Serialize(ToolArguments.ToSharedRef(), TJsonWriterFactory<>::Create(&PrettyJson)))
+	{
+		UE_LOG(LogModelContextProtocol, VeryVerbose, TEXT("'%s' parameters:\n%s"),  *ToolName, *PrettyJson);
+	}
+
+	FModelContextProtocolToolRequestId ToolRequestId(RequestId.ToSharedRef());
+
+	FString SessionId;
+	if (UE::ModelContextProtocol::Private::GetSessionIdFromRequestHeaders(Request, SessionId))
+	{
+		TSharedPtr<FModelContextProtocolSession> Session = FindSession(SessionId);
+
+		if (Session.IsValid())
+		{
+			TSharedPtr<FJsonValue> ProgressToken;
+
+			if (Params.IsValid())
+			{
+				const TSharedPtr<FJsonObject>* MetaParams = nullptr;
+
+				if (Params->TryGetObjectField(TEXT("_meta"), MetaParams))
+				{
+					if (MetaParams && MetaParams->IsValid())
+					{
+						ProgressToken = (*MetaParams)->TryGetField(TEXT("progressToken"));
+					}
+				}
+			}
+
+			FModelContextProtocolToolContext& Context = Session->ActiveRequests.Add(ToolRequestId);
+			Context.Tool = Tool;
+			Context.ProgressToken = ProgressToken;
+			// Saving the callback for re-use with SSE event streaming
+			Context.EventStreamWrite = OnComplete;
+			Context.LastProgressSeconds = FPlatformTime::Seconds();
+		}
+		else
+		{
+			UE::ModelContextProtocol::Private::CompleteWithError(OnComplete, RequestId, UE::ModelContextProtocol::Private::EJsonRpcErrorCode::InvalidRequest, FString::Printf(TEXT("Invalid session id")));
+			return true;
+		}
+	}
+	else
+	{
+		UE::ModelContextProtocol::Private::CompleteWithError(OnComplete, RequestId, UE::ModelContextProtocol::Private::EJsonRpcErrorCode::InvalidRequest, FString::Printf(TEXT("Missing session id")));
+		return true;
+	}
+
+	TUniquePtr<FHttpServerResponse> Response = FHttpServerResponse::Create(FString(TEXT("")), UE::ModelContextProtocol::ContentTypeEventStream);
+	Response->Headers.Add(TEXT("Connection"), { TEXT("keep-alive") });
+	Response->Headers.Add(TEXT("Cache-Control"), { TEXT("no-cache") });
+	Response->Headers.Add(UE::ModelContextProtocol::Private::McpSessionIdHeader, { SessionId });
+	EnumAddFlags(Response->Flags, EHttpServerResponseFlags::MultipleWriteStream | EHttpServerResponseFlags::HasAdditionalWrites);
+
+	OnComplete(MoveTemp(Response));
+
+	TWeakPtr<bool> WeakAlive = AliveGuard;
+	Tool->RunAsync(ToolRequestId, ToolArguments,
+		/*OnComplete*/[this, WeakAlive, RequestId, SessionId, ToolRequestId, OnComplete](const FModelContextProtocolToolResult& Result)
+		{
+			if (!WeakAlive.IsValid())
+			{
+				return;
+			}
+
+			FJsonDomBuilder::FObject ResponseObject;
+			ResponseObject.Set(TEXT("jsonrpc"), UE::ModelContextProtocol::JsonRpcVersion);
+			ResponseObject.Set(TEXT("id"), RequestId);
+
+			TSharedPtr<FModelContextProtocolSession> Session = FindSession(SessionId);
+
+			if (!Session.IsValid() || !Session->ActiveRequests.Contains(ToolRequestId))
+			{
+				// Request was cancelled or session was destroyed — silently discard per MCP spec.
+				return;
+			}
+
+			TSharedPtr<FJsonValue> ResultValue = MakeShared<FJsonValueObject>(Result.JsonObject);
+			ResponseObject.Set(TEXT("result"), ResultValue);
+			Session->ActiveRequests.Remove(ToolRequestId);
+
+			FUtf8String ResponseStr;
+			TSharedRef<TJsonWriter<UTF8CHAR, TCondensedJsonPrintPolicy<UTF8CHAR>>> JsonWriter = TJsonWriterFactory<UTF8CHAR, TCondensedJsonPrintPolicy<UTF8CHAR>>::Create(&ResponseStr);
+			FJsonSerializer::Serialize(ResponseObject.AsJsonObject(), JsonWriter);
+
+			TUniquePtr<FHttpServerResponse> ServerResponse = FHttpServerResponse::Create(UE::ModelContextProtocol::Private::FormatSSEMessage(ResponseStr), UE::ModelContextProtocol::ContentTypeEventStream);
+			EnumAddFlags(ServerResponse->Flags, EHttpServerResponseFlags::MultipleWriteStream | EHttpServerResponseFlags::SkipHeaderWrite);
+
+			OnComplete(MoveTemp(ServerResponse));
+		});
+
+	return true;
+}
+
+bool FModelContextProtocolServer::ProcessListResourcesJsonRpcCall(const FHttpServerRequest& Request, const TSharedPtr<FJsonValue>& RequestId, const TSharedPtr<FJsonObject>& Params, const FHttpResultCallback& OnComplete)
+{
+	const IModelContextProtocolModule& ModelContextProtocolModule = IModelContextProtocolModule::GetChecked();
+
+	LastResourceDescriptorList.Reset();
+	for (const TSharedRef<IModelContextProtocolResourceProvider>& ResourceProvider : ModelContextProtocolModule.GetResourceProviders())
+	{
+		ResourceProvider->ListResources(LastResourceDescriptorList);
+	}
+
+	UE_LOG(LogModelContextProtocol, Log, TEXT("Listing resources (%d)"), LastResourceDescriptorList.Num());
+
+	TSharedRef<FJsonValueArray> AllResourcesJsonValue = LastResourceDescriptorList.GetJsonArray();
+	const TArray<TSharedPtr<FJsonValue>>& AllResources = AllResourcesJsonValue->AsArray();
+	UE::ModelContextProtocol::Private::FPaginationResult PaginationResult;
+	FString PaginationError;
+	if (!UE::ModelContextProtocol::Private::ApplyPagination(AllResources, Params, PaginationResult, PaginationError))
+	{
+		UE::ModelContextProtocol::Private::CompleteWithError(OnComplete, RequestId, UE::ModelContextProtocol::Private::EJsonRpcErrorCode::InvalidParams, PaginationError);
+		LastResourceDescriptorList.ReleaseJsonArray();
+		return true;
+	}
+
+	FJsonDomBuilder::FObject ResultObject;
+	ResultObject.Set(TEXT("resources"), MakeShared<FJsonValueArray>(MoveTemp(PaginationResult.Items)));
+	if (PaginationResult.NextCursor.IsSet())
+	{
+		ResultObject.Set(TEXT("nextCursor"), PaginationResult.NextCursor.GetValue());
+	}
+
+	UE::ModelContextProtocol::Private::CompleteWithResult(OnComplete, RequestId, ResultObject);
+
+	LastResourceDescriptorList.ReleaseJsonArray();
+
+	return true;
+}
+
+bool FModelContextProtocolServer::ProcessReadResourceJsonRpcCall(const FHttpServerRequest& Request, const TSharedPtr<FJsonValue>& RequestId, const TSharedPtr<FJsonObject>& Params, const FHttpResultCallback& OnComplete)
+{
+	if (!Params.IsValid())
+	{
+		UE::ModelContextProtocol::Private::CompleteWithError(OnComplete, RequestId, UE::ModelContextProtocol::Private::EJsonRpcErrorCode::InvalidParams, TEXT("Expected params for call to resources/read"));
+		return true;
+	}
+
+	FString ResourceUri;
+	if (!Params->TryGetStringField(TEXT("uri"), ResourceUri))
+	{
+		UE::ModelContextProtocol::Private::CompleteWithError(OnComplete, RequestId, UE::ModelContextProtocol::Private::EJsonRpcErrorCode::InvalidParams, TEXT("Expected uri param for call to resources/read"));
+		return true;
+	}
+
+	// Look up which resource provider can serve this resource. Check the cache first (populated by resources/list),
+	// then fall back to querying all providers directly, since clients may call resources/read without prior resources/list.
+	TSharedPtr<const IModelContextProtocolResourceProvider> ResourceProvider = LastResourceDescriptorList.FindResourceProvider(ResourceUri);
+	if (!ResourceProvider.IsValid())
+	{
+		const IModelContextProtocolModule& ModelContextProtocolModule = IModelContextProtocolModule::GetChecked();
+		for (const TSharedRef<IModelContextProtocolResourceProvider>& Provider : ModelContextProtocolModule.GetResourceProviders())
+		{
+			FModelContextProtocolResourceDescriptorList ProviderResources;
+			Provider->ListResources(ProviderResources);
+			if (ProviderResources.FindResourceProvider(ResourceUri).IsValid())
+			{
+				ResourceProvider = Provider;
+				break;
+			}
+		}
+	}
+
+	if (!ResourceProvider.IsValid())
+	{
+		UE::ModelContextProtocol::Private::CompleteWithError(OnComplete, RequestId, UE::ModelContextProtocol::Private::EJsonRpcErrorCode::ResourceNotFound, FString::Printf(TEXT("Resource not found: %s"), *ResourceUri));
+		return true;
+	}
+
+	TValueOrError<FModelContextProtocolResource, FString> ResourceOrError = ResourceProvider->ReadResource(ResourceUri);
+	if (ResourceOrError.HasError())
+	{
+		UE::ModelContextProtocol::Private::CompleteWithError(OnComplete, RequestId, UE::ModelContextProtocol::Private::EJsonRpcErrorCode::InternalError, FString::Printf(TEXT("Error reading resource: %s"), *ResourceOrError.GetError()));
+		return true;
+	}
+
+	FJsonDomBuilder::FObject ResultObject;
+	ResultObject.Set(TEXT("contents"), FJsonDomBuilder::FArray().Add(ResourceOrError.GetValue().GetJsonObject()));
+
+	UE::ModelContextProtocol::Private::CompleteWithResult(OnComplete, RequestId, ResultObject);
+	return true;
+}
+
+void FModelContextProtocolServer::Tick(float DeltaTime)
+{
+	if (UE::ModelContextProtocol::ProgressIntervalSeconds > 0.0f)
+	{
+		const double NowSeconds = FPlatformTime::Seconds();
+
+		for (TSharedPtr<FModelContextProtocolSession>& Session : Sessions)
+		{
+			if (Session.IsValid())
+			{
+				for (TPair<FModelContextProtocolToolRequestId, FModelContextProtocolToolContext>& ToolRequest : Session->ActiveRequests)
+				{
+					FModelContextProtocolToolContext& Context = ToolRequest.Value;
+
+					if (Context.ProgressToken.IsValid() && Context.EventStreamWrite.IsSet())
+					{
+						if ((NowSeconds - Context.LastProgressSeconds) >= UE::ModelContextProtocol::ProgressIntervalSeconds)
+						{
+							// Not sending a total progress value, so this is simply a heartbeat as the total duration is unknown
+							++Context.LastProgressValue;
+
+							UE::ModelContextProtocol::Private::SendProgressUpdate(Context.EventStreamWrite, Context.ProgressToken.ToSharedRef(), Context.LastProgressValue);
+
+							Context.LastProgressSeconds = NowSeconds;
+						}
+					}
+				}
+			}
+		}
+	}
+}
+
+bool FModelContextProtocolServer::ProcessGetRequest(const FHttpServerRequest& Request, const FHttpResultCallback& OnComplete)
+{
+	if (!UE::ModelContextProtocol::Private::ValidateOriginHeader(Request, OnComplete))
+	{
+		return true;
+	}
+
+	// We do not currently support sse on a separate endpoint
+	UE::ModelContextProtocol::Private::CompleteWithResponseCode(OnComplete, EHttpServerResponseCodes::BadMethod);
+	return true;
+}
+
+bool FModelContextProtocolServer::ProcessDeleteRequest(const FHttpServerRequest& Request, const FHttpResultCallback& OnComplete)
+{
+	if (!UE::ModelContextProtocol::Private::ValidateOriginHeader(Request, OnComplete))
+	{
+		return true;
+	}
+
+	FString SessionId;
+	if (UE::ModelContextProtocol::Private::GetSessionIdFromRequestHeaders(Request, SessionId))
+	{
+		const int32 RemovedCount = Sessions.RemoveAll([&SessionId](const TSharedPtr<FModelContextProtocolSession>& Session)
+			{
+				return Session.IsValid() && (Session->ID == SessionId);
+			});
+
+		if (RemovedCount > 0)
+		{
+			UE::ModelContextProtocol::Private::CompleteWithResponseCode(OnComplete, EHttpServerResponseCodes::Accepted);
+		}
+		else
+		{
+			UE::ModelContextProtocol::Private::CompleteWithResponseCode(OnComplete, EHttpServerResponseCodes::BadRequest);
+		}
+	}
+	else
+	{
+		UE::ModelContextProtocol::Private::CompleteWithResponseCode(OnComplete, EHttpServerResponseCodes::BadRequest);
+	}
+
+	return true;
+}
diff --git a/Engine/Plugins/Experimental/ModelContextProtocol/Source/ModelContextProtocol/Private/ModelContextProtocolSettings.cpp b/Engine/Plugins/Experimental/ModelContextProtocol/Source/ModelContextProtocol/Private/ModelContextProtocolSettings.cpp
deleted file mode 100644
diff --git a/Engine/Plugins/Experimental/ModelContextProtocol/Source/ModelContextProtocol/Private/ModelContextProtocolToolAsyncAction.cpp b/Engine/Plugins/Experimental/ModelContextProtocol/Source/ModelContextProtocol/Private/ModelContextProtocolToolAsyncAction.cpp
deleted file mode 100644
diff --git a/Engine/Plugins/Experimental/ModelContextProtocol/Source/ModelContextProtocol/Private/ModelContextProtocolToolLibrary.cpp b/Engine/Plugins/Experimental/ModelContextProtocol/Source/ModelContextProtocol/Private/ModelContextProtocolToolLibrary.cpp
deleted file mode 100644
diff --git a/Engine/Plugins/Experimental/ModelContextProtocol/Source/ModelContextProtocol/Private/ModelContextProtocolToolResults.cpp b/Engine/Plugins/Experimental/ModelContextProtocol/Source/ModelContextProtocol/Private/ModelContextProtocolToolResults.cpp
--- a/Engine/Plugins/Experimental/ModelContextProtocol/Source/ModelContextProtocol/Private/ModelContextProtocolToolResults.cpp
+++ b/Engine/Plugins/Experimental/ModelContextProtocol/Source/ModelContextProtocol/Private/ModelContextProtocolToolResults.cpp
@@ -3,16 +3,11 @@
 #include "ModelContextProtocolToolResults.h"
 #include "ModelContextProtocol.h"
 
-#include "Engine/Texture2D.h"
-#include "ImageUtils.h"
 #include "JsonObjectConverter.h"
 #include "ModelContextProtocolResources.h"
 #include "Misc/Base64.h"
-#include "Misc/ScopeExit.h"
 #include "Policies/CondensedJsonPrintPolicy.h"
 #include "Serialization/JsonWriter.h"
-#include "TextureResource.h"
-#include "Sound/SoundWave.h"
 
 #include UE_INLINE_GENERATED_CPP_BY_NAME(ModelContextProtocolToolResults)
 
@@ -42,49 +37,6 @@ namespace UE::ModelContextProtocol::Private
 	}
 }
 
-FModelContextProtocolToolResult UE::ModelContextProtocol::MakeImageResult(const UTexture2D* Texture, const TCHAR* ToFormatExtension, EModelContextProtocolAudience Audience)
-{
-	if (!Texture)
-	{
-		return FModelContextProtocolToolResult();
-	}
-
-	if (const FTexturePlatformData* TexturePlatformData = Texture->GetPlatformData();
-		ensure(TexturePlatformData) && ensure(!TexturePlatformData->Mips.IsEmpty()) && ensure(TexturePlatformData->PixelFormat == PF_B8G8R8A8))
-	{
-		const FTexture2DMipMap& Mip0 = TexturePlatformData->Mips[0];
-		if (Mip0.SizeX > 0 && Mip0.SizeY > 0)
-		{
-			const void* Mip0Data = Mip0.BulkData.LockReadOnly();
-			check(Mip0Data);
-			ON_SCOPE_EXIT{ Mip0.BulkData.Unlock(); };
-
-			return MakeImageResult(FImageView(static_cast<const FColor*>(Mip0Data), Mip0.SizeX, Mip0.SizeY), ToFormatExtension, Audience);
-		}
-	}
-
-	return FModelContextProtocolToolResult();
-}
-
-FModelContextProtocolToolResult UE::ModelContextProtocol::MakeImageResult(const FImageView& Image, const TCHAR* ToFormatExtension, EModelContextProtocolAudience Audience)
-{
-	if (!ensure(Image.IsImageInfoValid()))
-	{
-		return FModelContextProtocolToolResult();
-	}
-	
-	TArray64<uint8> CompressedData;
-	FImageUtils::CompressImage(CompressedData,ToFormatExtension, Image);
-
-	FStringView Format(ToFormatExtension);
-	if (Format.StartsWith(TEXT('.')))
-	{
-		Format.RightChopInline(1);
-	}
-	
-	return MakeImageResult(FString(TEXT("image/")) + Format, CompressedData, Audience);
-}
-
 FModelContextProtocolToolResult UE::ModelContextProtocol::MakeImageResult(const FString& MimeType, TArrayView<uint8> Data, EModelContextProtocolAudience Audience)
 {
 	FString Base64EncodedData = FBase64::Encode(Data.GetData(), Data.NumBytes());
@@ -121,37 +73,6 @@ FModelContextProtocolToolResult UE::ModelContextProtocol::MakeImageResult(const
 	return ResultObject.AsJsonObject().ToSharedPtr();
 }
 
-FModelContextProtocolToolResult UE::ModelContextProtocol::MakeAudioResult(USoundWave* SoundWave)
-{
-	if (!SoundWave)
-	{
-		return FModelContextProtocolToolResult();
-	}
-
-#if WITH_EDITOR
-	if (UE::ModelContextProtocol::bAudioResultOggFormat)
-	{
-		if (FByteBulkData* BulkData = SoundWave->GetCompressedData("OGG"))
-		{
-			TArray<uint8> RawData;
-			RawData.Append(BulkData->GetCopyAsBuffer(BulkData->GetElementCount(), true).GetView());
-
-			return MakeAudioResult(TEXT("audio/ogg"), RawData);
-		}
-	}
-	else
-	{
-		TFuture<FSharedBuffer> FutureBuffer = SoundWave->RawData.GetPayload(USoundWave::FEditorAudioBulkData::EPayloadFlags::ReturnCopyWithTransformationData); // This will block
-		if (FutureBuffer.IsReady())
-		{
-			return MakeAudioResult(TEXT("audio/wav"), MakeArrayView<uint8>((uint8*)FutureBuffer.Get().GetData(), FutureBuffer.Get().GetSize()));
-		}
-	}
-#endif // WITH_EDITOR
-	
-	return FModelContextProtocolToolResult();
-}
-
 FModelContextProtocolToolResult UE::ModelContextProtocol::MakeAudioResult(const FString& MimeType, TArrayView<uint8> Data)
 {
 	FString Base64EncodedData = FBase64::Encode(Data.GetData(), Data.NumBytes());
diff --git a/Engine/Plugins/Experimental/ModelContextProtocol/Source/ModelContextProtocol/Private/ModelContextProtocolToolUtils.cpp b/Engine/Plugins/Experimental/ModelContextProtocol/Source/ModelContextProtocol/Private/ModelContextProtocolToolUtils.cpp
deleted file mode 100644
diff --git a/Engine/Plugins/Experimental/ModelContextProtocol/Source/ModelContextProtocol/Public/IModelContextProtocolModule.h b/Engine/Plugins/Experimental/ModelContextProtocol/Source/ModelContextProtocol/Public/IModelContextProtocolModule.h
--- a/Engine/Plugins/Experimental/ModelContextProtocol/Source/ModelContextProtocol/Public/IModelContextProtocolModule.h
+++ b/Engine/Plugins/Experimental/ModelContextProtocol/Source/ModelContextProtocol/Public/IModelContextProtocolModule.h
@@ -2,8 +2,10 @@
 
 #pragma once
 
+#include "ModelContextProtocol.h"
 #include "Modules/ModuleManager.h"
 
+class FModelContextProtocolServer;
 struct IModelContextProtocolResourceProvider;
 struct IModelContextProtocolTool;
 
@@ -55,6 +57,15 @@ public:
 	 */
 	virtual void RefreshTools() = 0;
 
+	/** Returns the MCP server instance, or nullptr if not created. */
+	[[nodiscard]] virtual FModelContextProtocolServer* GetServer() = 0;
+
+	/** Creates and starts the MCP HTTP server on the specified port and URL path. */
+	virtual void StartServer(uint32 Port = UE::ModelContextProtocol::DefaultServerPort, const FString& UrlPath = UE::ModelContextProtocol::DefaultServerUrlPath) = 0;
+
+	/** Stops and destroys the MCP HTTP server. */
+	virtual void StopServer() = 0;
+
 	/**
 	 * Returns current list of registered resource providers.
 	 * @see AddResourceProvider, FindResourceProvider
diff --git a/Engine/Plugins/Experimental/ModelContextProtocol/Source/ModelContextProtocol/Public/IModelContextProtocolResourceProvider.h b/Engine/Plugins/Experimental/ModelContextProtocol/Source/ModelContextProtocol/Public/IModelContextProtocolResourceProvider.h
--- a/Engine/Plugins/Experimental/ModelContextProtocol/Source/ModelContextProtocol/Public/IModelContextProtocolResourceProvider.h
+++ b/Engine/Plugins/Experimental/ModelContextProtocol/Source/ModelContextProtocol/Public/IModelContextProtocolResourceProvider.h
@@ -2,7 +2,6 @@
 
 #pragma once
 
-#include "Engine/EngineBaseTypes.h"
 #include "Templates/SharedPointer.h"
 #include "Templates/ValueOrError.h"
 
diff --git a/Engine/Plugins/Experimental/ModelContextProtocol/Source/ModelContextProtocol/Public/ModelContextProtocol.h b/Engine/Plugins/Experimental/ModelContextProtocol/Source/ModelContextProtocol/Public/ModelContextProtocol.h
--- a/Engine/Plugins/Experimental/ModelContextProtocol/Source/ModelContextProtocol/Public/ModelContextProtocol.h
+++ b/Engine/Plugins/Experimental/ModelContextProtocol/Source/ModelContextProtocol/Public/ModelContextProtocol.h
@@ -42,14 +42,33 @@ namespace UE::ModelContextProtocol
 		return ProtocolVersion;
 	}
 
+	static constexpr uint32 DefaultServerPort = 8000;
+	static constexpr const TCHAR* DefaultServerUrlPath = TEXT("/mcp");
+
+	/** Tool name validation result per MCP spec (2025-11-25). */
+	enum class EToolNameValidation : uint8
+	{
+		Valid,
+		Empty,
+		ExceedsMaxLength,
+		InvalidCharacters
+	};
+
+	/**
+	 * Validates a tool name per MCP spec (2025-11-25) tool name rules.
+	 * Names SHOULD be 1-128 chars using only A-Z, a-z, 0-9, underscore, hyphen, and dot.
+	 * @see https://modelcontextprotocol.io/specification/2025-11-25/server/tools#tool-names
+	 */
+	UE_API EToolNameValidation ValidateToolName(const FString& ToolName);
+
 	static constexpr const TCHAR* PODWrapperResultPropertyName = TEXT("result");
 
 	static constexpr const TCHAR* ContentTypeEventStream = TEXT("text/event-stream");
 
 	// CVars
-	extern bool bWrapPODResultsInObject; 
-	extern bool bAudioResultOggFormat;
-	extern float ProgressIntervalSeconds;
+	UE_API extern bool bWrapPODResultsInObject;
+	UE_API extern bool bAudioResultOggFormat;
+	UE_API extern float ProgressIntervalSeconds;
 	UE_API extern int32 PaginationPageSize;
 }
 
diff --git a/Engine/Plugins/Experimental/ModelContextProtocol/Source/ModelContextProtocol/Public/ModelContextProtocolEngineSubsystem.h b/Engine/Plugins/Experimental/ModelContextProtocol/Source/ModelContextProtocol/Public/ModelContextProtocolEngineSubsystem.h
deleted file mode 100644
diff --git a/Engine/Plugins/Experimental/ModelContextProtocol/Source/ModelContextProtocol/Public/ModelContextProtocolMetaData.h b/Engine/Plugins/Experimental/ModelContextProtocol/Source/ModelContextProtocol/Public/ModelContextProtocolMetaData.h
--- a/Engine/Plugins/Experimental/ModelContextProtocol/Source/ModelContextProtocol/Public/ModelContextProtocolMetaData.h
+++ b/Engine/Plugins/Experimental/ModelContextProtocol/Source/ModelContextProtocol/Public/ModelContextProtocolMetaData.h
@@ -56,24 +56,3 @@ struct FModelContextProtocolFunctionMetaData
 	FString WorldContext;
 };
 
-namespace UE::ModelContextProtocol
-{
-#if WITH_EDITORONLY_DATA
-	/**
-	 * Collects editor-only metadata for a UFunction, storing it in cookable form.
-	 *
-	 * @param Function The function to collect metadata for.
-	 * @param PropertyFilter Property filter for schema generation.
-	 * @return Collected function metadata.
-	 */
-	MODELCONTEXTPROTOCOL_API FModelContextProtocolFunctionMetaData CollectFunctionMetaData(const UFunction* Function, const FJsonSchemaPropertyFilter& PropertyFilter);
-#endif
-
-	/**
-	 * Converts cookable function metadata back into FJsonSchemaEditorMetadata for passing to FJsonSchemaGenerator.
-	 *
-	 * @param FunctionMetaData The cookable metadata to convert.
-	 * @return Engine editor metadata suitable for schema generation.
-	 */
-	MODELCONTEXTPROTOCOL_API FJsonSchemaEditorMetadata ConvertToCachedEditorMetadata(const FModelContextProtocolFunctionMetaData& FunctionMetaData);
-}
diff --git a/Engine/Plugins/Experimental/ModelContextProtocol/Source/ModelContextProtocol/Public/ModelContextProtocolResources.h b/Engine/Plugins/Experimental/ModelContextProtocol/Source/ModelContextProtocol/Public/ModelContextProtocolResources.h
--- a/Engine/Plugins/Experimental/ModelContextProtocol/Source/ModelContextProtocol/Public/ModelContextProtocolResources.h
+++ b/Engine/Plugins/Experimental/ModelContextProtocol/Source/ModelContextProtocol/Public/ModelContextProtocolResources.h
@@ -2,7 +2,6 @@
 
 #pragma once
 
-#include "Engine/EngineBaseTypes.h"
 #include "JsonDomBuilder.h"
 #include "Misc/Base64.h"
 #include "Templates/SharedPointer.h"
diff --git a/Engine/Plugins/Experimental/ModelContextProtocol/Source/ModelContextProtocol/Public/ModelContextProtocolServer.h b/Engine/Plugins/Experimental/ModelContextProtocol/Source/ModelContextProtocol/Public/ModelContextProtocolServer.h
new file mode 100644
--- /dev/null
+++ b/Engine/Plugins/Experimental/ModelContextProtocol/Source/ModelContextProtocol/Public/ModelContextProtocolServer.h
@@ -0,0 +1,86 @@
+// Copyright Epic Games, Inc. All Rights Reserved.
+
+#pragma once
+
+#include "ModelContextProtocolResources.h"
+#include "ModelContextProtocolSession.h"
+
+#include "Containers/Ticker.h"
+#include "HttpResultCallback.h"
+#include "HttpRouteHandle.h"
+#include "Templates/SharedPointer.h"
+
+class IHttpRouter;
+class FJsonValue;
+class FJsonObject;
+
+/**
+ * MCP (Anthropic's Model Context Protocol) tool server. Serves MCP tools over HTTP per the https://modelcontextprotocol.io/ spec.
+ * @see https://modelcontextprotocol.io
+ * @see UModelContextProtocolSettings
+ */
+class FModelContextProtocolServer
+{
+public:
+
+	FModelContextProtocolServer() = default;
+	MODELCONTEXTPROTOCOL_API ~FModelContextProtocolServer();
+
+	FModelContextProtocolServer(const FModelContextProtocolServer&) = delete;
+	FModelContextProtocolServer& operator=(const FModelContextProtocolServer&) = delete;
+
+	/**
+	 * Starts the HTTP MCP server on the specified port and URL path.
+	 */
+	MODELCONTEXTPROTOCOL_API void StartServer(uint32 Port, const FString& UrlPath);
+
+	/**
+	 * Stops the HTTP MCP server.
+	 */
+	MODELCONTEXTPROTOCOL_API void StopServer();
+
+	/** Returns true if the HTTP MCP server is currently running. */
+	MODELCONTEXTPROTOCOL_API bool IsServerRunning() const;
+
+	/** Returns the port the server is currently listening on, or 0 if not running. */
+	MODELCONTEXTPROTOCOL_API uint32 GetServerPort() const;
+
+	/** Sends notifications/tools/list_changed to all initialized sessions with active SSE streams. */
+	MODELCONTEXTPROTOCOL_API void BroadcastToolsListChanged();
+
+private:
+
+	void Tick(float DeltaTime);
+
+	TArray<TSharedPtr<FModelContextProtocolSession>> Sessions;
+
+	TSharedPtr<FModelContextProtocolSession> FindSession(const FString& SessionId) const;
+
+	TSharedPtr<IHttpRouter> HttpRouter;
+	uint32 ActiveServerPort = 0;
+
+	FHttpRouteHandle MainMcpRoute;
+	FHttpRouteHandle SseMcpRoute;
+	FHttpRouteHandle DeleteMcpRoute;
+
+	FTSTicker::FDelegateHandle TickerHandle;
+
+	/** Prevent use-after-free: async tool completion lambdas capture a TWeakPtr to this; Reset() in destructor invalidates all weak references. */
+	TSharedPtr<bool> AliveGuard = MakeShared<bool>(true);
+
+	FModelContextProtocolResourceDescriptorList LastResourceDescriptorList;
+
+	bool ProcessPostRequest(const FHttpServerRequest& Request, const FHttpResultCallback& OnComplete);
+	bool ProcessGetRequest(const FHttpServerRequest& Request, const FHttpResultCallback& OnComplete);
+	bool ProcessDeleteRequest(const FHttpServerRequest& Request, const FHttpResultCallback& OnComplete);
+
+	bool ProcessJsonRpcCall(const FHttpServerRequest& Request, const TSharedPtr<FJsonValue>& RequestId, const FString& Method, const TSharedPtr<FJsonObject>& Params, const FHttpResultCallback& OnComplete);
+	bool ProcessPingJsonRpcCall(const FHttpServerRequest& Request, const TSharedPtr<FJsonValue>& RequestId, const FHttpResultCallback& OnComplete);
+	bool ProcessInitializeJsonRpcCall(const FHttpServerRequest& Request, const TSharedPtr<FJsonValue>& RequestId, const TSharedPtr<FJsonObject>& Params, const FHttpResultCallback& OnComplete);
+	bool ProcessInitializedNotificationJsonRpcCall(const FHttpServerRequest& Request, const TSharedPtr<FJsonValue>& RequestId, const TSharedPtr<FJsonObject>& Params, const FHttpResultCallback& OnComplete);
+	bool ProcessNotificationCancelledJsonRpcCall(const FHttpServerRequest& Request, const TSharedPtr<FJsonValue>& RequestId, const TSharedPtr<FJsonObject>& Params, const FHttpResultCallback& OnComplete);
+	bool ProcessListToolsJsonRpcCall(const FHttpServerRequest& Request, const TSharedPtr<FJsonValue>& RequestId, const TSharedPtr<FJsonObject>& Params, const FHttpResultCallback& OnComplete);
+	bool ProcessToolCallJsonRpcCall(const FHttpServerRequest& Request, const TSharedPtr<FJsonValue>& RequestId, const TSharedPtr<FJsonObject>& Params, const FHttpResultCallback& OnComplete);
+	bool ProcessListResourcesJsonRpcCall(const FHttpServerRequest& Request, const TSharedPtr<FJsonValue>& RequestId, const TSharedPtr<FJsonObject>& Params, const FHttpResultCallback& OnComplete);
+	bool ProcessReadResourceJsonRpcCall(const FHttpServerRequest& Request, const TSharedPtr<FJsonValue>& RequestId, const TSharedPtr<FJsonObject>& Params, const FHttpResultCallback& OnComplete);
+};
diff --git a/Engine/Plugins/Experimental/ModelContextProtocol/Source/ModelContextProtocol/Public/ModelContextProtocolSettings.h b/Engine/Plugins/Experimental/ModelContextProtocol/Source/ModelContextProtocol/Public/ModelContextProtocolSettings.h
deleted file mode 100644
diff --git a/Engine/Plugins/Experimental/ModelContextProtocol/Source/ModelContextProtocol/Public/ModelContextProtocolToolAsyncAction.h b/Engine/Plugins/Experimental/ModelContextProtocol/Source/ModelContextProtocol/Public/ModelContextProtocolToolAsyncAction.h
deleted file mode 100644
diff --git a/Engine/Plugins/Experimental/ModelContextProtocol/Source/ModelContextProtocol/Public/ModelContextProtocolToolLibrary.h b/Engine/Plugins/Experimental/ModelContextProtocol/Source/ModelContextProtocol/Public/ModelContextProtocolToolLibrary.h
deleted file mode 100644
diff --git a/Engine/Plugins/Experimental/ModelContextProtocol/Source/ModelContextProtocol/Public/ModelContextProtocolToolResults.h b/Engine/Plugins/Experimental/ModelContextProtocol/Source/ModelContextProtocol/Public/ModelContextProtocolToolResults.h
--- a/Engine/Plugins/Experimental/ModelContextProtocol/Source/ModelContextProtocol/Public/ModelContextProtocolToolResults.h
+++ b/Engine/Plugins/Experimental/ModelContextProtocol/Source/ModelContextProtocol/Public/ModelContextProtocolToolResults.h
@@ -11,10 +11,6 @@
 
 #define UE_API MODELCONTEXTPROTOCOL_API
 
-class USoundWave;
-class UTexture2D;
-struct FColor;
-struct FImageView;
 struct FModelContextProtocolResourceDescriptor;
 
 UENUM(meta = (Bitflags))
@@ -82,13 +78,10 @@ namespace UE::ModelContextProtocol
 	}
 
 	// @see https://modelcontextprotocol.io/specification/2025-06-18/server/tools#image-content
-	UE_API FModelContextProtocolToolResult MakeImageResult(const UTexture2D* Texture, const TCHAR* ToFormatExtension = TEXT("jpeg"), EModelContextProtocolAudience Audience = EModelContextProtocolAudience::All);
-	UE_API FModelContextProtocolToolResult MakeImageResult(const FImageView& Image, const TCHAR* ToFormatExtension = TEXT("jpeg"), EModelContextProtocolAudience Audience = EModelContextProtocolAudience::All);
 	UE_API FModelContextProtocolToolResult MakeImageResult(const FString& MimeType, TArrayView<uint8> Data, EModelContextProtocolAudience Audience = EModelContextProtocolAudience::All);
 	UE_API FModelContextProtocolToolResult MakeImageResult(const FString& MimeType, FString&& Base64EncodedData, EModelContextProtocolAudience Audience = EModelContextProtocolAudience::All);
 
 	// @see https://modelcontextprotocol.io/specification/2025-06-18/server/tools#audio-content
-	UE_API FModelContextProtocolToolResult MakeAudioResult(USoundWave* SoundWave);
 	UE_API FModelContextProtocolToolResult MakeAudioResult(const FString& MimeType, TArrayView<uint8> Data);
 	UE_API FModelContextProtocolToolResult MakeAudioResult(const FString& MimeType, FString&& Base64EncodedData);
 
diff --git a/Engine/Plugins/Experimental/ModelContextProtocol/Source/ModelContextProtocol/Public/ModelContextProtocolToolUtils.h b/Engine/Plugins/Experimental/ModelContextProtocol/Source/ModelContextProtocol/Public/ModelContextProtocolToolUtils.h
deleted file mode 100644
diff --git a/Engine/Plugins/Experimental/ModelContextProtocol/Source/ModelContextProtocolEditor/ModelContextProtocolEditor.Build.cs b/Engine/Plugins/Experimental/ModelContextProtocol/Source/ModelContextProtocolEditor/ModelContextProtocolEditor.Build.cs
--- a/Engine/Plugins/Experimental/ModelContextProtocol/Source/ModelContextProtocolEditor/ModelContextProtocolEditor.Build.cs
+++ b/Engine/Plugins/Experimental/ModelContextProtocol/Source/ModelContextProtocolEditor/ModelContextProtocolEditor.Build.cs
@@ -7,17 +7,18 @@ public class ModelContextProtocolEditor : ModuleRules
 	public ModelContextProtocolEditor(ReadOnlyTargetRules Target) : base(Target)
 	{
 		PCHUsage = ModuleRules.PCHUsageMode.UseExplicitOrSharedPCHs;
-		
+
 		PublicDependencyModuleNames.AddRange(
 			new string[]
 			{
-				"Core", 
+				"Core",
 				"AssetDefinition",
 				"EngineAssetDefinitions",
-				"ModelContextProtocol", 
+				"ModelContextProtocol",
+				"ModelContextProtocolEngine",
 				"Blutility"
 			});
-			
+
 		PrivateDependencyModuleNames.AddRange(
 			new string[]
 			{
@@ -28,10 +29,15 @@ public class ModelContextProtocolEditor : ModuleRules
 				"ToolMenus",
 				"CoreUObject",
 				"Engine",
+				"Json",
+				"JsonUtilities",
 				"Slate",
-				"SlateCore", 
-				"ModelContextProtocol",
-				"KismetCompiler"
+				"SlateCore",
+				"BlueprintGraph",
+				"DeveloperSettings",
+				"JsonUtilitiesEditor",
+				"KismetCompiler",
+				"ToolsetRegistry"
 			});
 	}
 }
diff --git a/Engine/Plugins/Experimental/ModelContextProtocol/Source/ModelContextProtocolEditor/Private/ModelContextProtocolEditor.cpp b/Engine/Plugins/Experimental/ModelContextProtocol/Source/ModelContextProtocolEditor/Private/ModelContextProtocolEditor.cpp
--- a/Engine/Plugins/Experimental/ModelContextProtocol/Source/ModelContextProtocolEditor/Private/ModelContextProtocolEditor.cpp
+++ b/Engine/Plugins/Experimental/ModelContextProtocol/Source/ModelContextProtocolEditor/Private/ModelContextProtocolEditor.cpp
@@ -1,22 +1,78 @@
 ﻿// Copyright Epic Games, Inc. All Rights Reserved.
 
 #include "ModelContextProtocolEditor.h"
+#include "IModelContextProtocolModule.h"
 #include "ModelContextProtocolEditorToolLibrary.h"
+#include "ModelContextProtocolSettings.h"
+#include "ModelContextProtocolToolLibrary.h"
 
+#include "Editor.h"
 #include "KismetCompilerModule.h"
+#include "Misc/CoreDelegates.h"
+#include "ToolsetRegistry/ToolsetRegistry.h"
+#include "ToolsetRegistry/ToolsetRegistrySubsystem.h"
 
 #define LOCTEXT_NAMESPACE "FModelContextProtocolEditorModule"
 
 void FModelContextProtocolEditorModule::StartupModule()
 {
 	IKismetCompilerInterface& KismetCompilerModule = FModuleManager::LoadModuleChecked<IKismetCompilerInterface>("KismetCompiler");
+	KismetCompilerModule.OverrideBPTypeForClass(UModelContextProtocolToolLibrary::StaticClass(), UModelContextProtocolToolLibraryBlueprint::StaticClass());
 	KismetCompilerModule.OverrideBPTypeForClass(UModelContextProtocolEditorToolLibrary::StaticClass(), UModelContextProtocolEditorToolLibraryBlueprint::StaticClass());
+
+	// Re-register adapters when MCP tools are refreshed
+	if (IModelContextProtocolModule* Module = IModelContextProtocolModule::Get())
+	{
+		OnRefreshToolsHandle = Module->OnRefreshTools().AddLambda([this]()
+		{
+			ToolsetRegistryAdapterManager.RegisterTools();
+		});
+	}
+
+	// Defer ToolsetRegistry binding and auto-start until the editor is fully initialized
+	// (GEditor and editor subsystems are not available during StartupModule)
+	PostEngineInitHandle = FCoreDelegates::GetOnPostEngineInit().AddLambda([this]()
+	{
+		// Subscribe to ToolsetRegistry changes so adapters are re-registered when toolsets are added
+		if (UToolsetRegistrySubsystem* ToolsetRegistrySubsystem = GEditor ? GEditor->GetEditorSubsystem<UToolsetRegistrySubsystem>() : nullptr)
+		{
+			OnToolsetRegisteredHandle = ToolsetRegistrySubsystem->ToolsetRegistry.OnToolsetRegistered().AddLambda([this]()
+			{
+				ToolsetRegistryAdapterManager.RegisterTools();
+			});
+		}
+
+		// Initial registration of any toolsets already loaded
+		ToolsetRegistryAdapterManager.RegisterTools();
+
+		// Auto-start the MCP server if configured
+		if (UE::ModelContextProtocol::ShouldAutoStartServer())
+		{
+			if (IModelContextProtocolModule* Module = IModelContextProtocolModule::Get())
+			{
+				Module->StartServer(UE::ModelContextProtocol::GetServerPortNumber(), UE::ModelContextProtocol::GetServerUrlPath());
+			}
+		}
+	});
 }
 
 void FModelContextProtocolEditorModule::ShutdownModule()
 {
+	FCoreDelegates::GetOnPostEngineInit().Remove(PostEngineInitHandle);
+
+	if (IModelContextProtocolModule* Module = IModelContextProtocolModule::Get())
+	{
+		Module->OnRefreshTools().Remove(OnRefreshToolsHandle);
+	}
+
+	if (UToolsetRegistrySubsystem* ToolsetRegistrySubsystem = GEditor ? GEditor->GetEditorSubsystem<UToolsetRegistrySubsystem>() : nullptr)
+	{
+		ToolsetRegistrySubsystem->ToolsetRegistry.OnToolsetRegistered().Remove(OnToolsetRegisteredHandle);
+	}
+
+	ToolsetRegistryAdapterManager.DeregisterTools();
 }
 
 #undef LOCTEXT_NAMESPACE
-	
-IMPLEMENT_MODULE(FModelContextProtocolEditorModule, ModelContextProtocolEditor)
\ No newline at end of file
+
+IMPLEMENT_MODULE(FModelContextProtocolEditorModule, ModelContextProtocolEditor)
diff --git a/Engine/Plugins/Experimental/ModelContextProtocol/Source/ModelContextProtocolEditor/Private/ModelContextProtocolMetaData.cpp b/Engine/Plugins/Experimental/ModelContextProtocol/Source/ModelContextProtocolEditor/Private/ModelContextProtocolMetaData.cpp
new file mode 100644
--- /dev/null
+++ b/Engine/Plugins/Experimental/ModelContextProtocol/Source/ModelContextProtocolEditor/Private/ModelContextProtocolMetaData.cpp
@@ -0,0 +1,5 @@
+// Copyright Epic Games, Inc. All Rights Reserved.
+
+// This file is intentionally empty.
+// CollectFunctionMetaData and ConvertToCachedEditorMetadata implementations
+// have been moved to ModelContextProtocolEngine module (ModelContextProtocolEngineMetaData.cpp).
diff --git a/Engine/Plugins/Experimental/ModelContextProtocol/Source/ModelContextProtocolEditor/Private/ModelContextProtocolToolsetRegistryAdapter.cpp b/Engine/Plugins/Experimental/ModelContextProtocol/Source/ModelContextProtocolEditor/Private/ModelContextProtocolToolsetRegistryAdapter.cpp
new file mode 100644
--- /dev/null
+++ b/Engine/Plugins/Experimental/ModelContextProtocol/Source/ModelContextProtocolEditor/Private/ModelContextProtocolToolsetRegistryAdapter.cpp
@@ -0,0 +1,157 @@
+// Copyright Epic Games, Inc. All Rights Reserved.
+
+#include "ModelContextProtocolToolsetRegistryAdapter.h"
+#include "IModelContextProtocolModule.h"
+#include "ModelContextProtocol.h"
+#include "ModelContextProtocolToolResults.h"
+
+#include "Editor.h"
+#include "Serialization/JsonReader.h"
+#include "Serialization/JsonSerializer.h"
+#include "ToolsetRegistry/ToolsetRegistry.h"
+#include "ToolsetRegistry/ToolsetRegistrySubsystem.h"
+
+FToolsetRegistryToolAdapter::FToolsetRegistryToolAdapter(const FString& InToolName, const FString& InDescription, TSharedPtr<FJsonObject> InInputSchema)
+	: ToolName(InToolName)
+	, Description(InDescription)
+	, InputSchema(MoveTemp(InInputSchema))
+{
+}
+
+void FToolsetRegistryToolAdapter::RunAsync(const FModelContextProtocolToolRequestId& RequestId, const TSharedPtr<FJsonObject>& Params, const FResultCallback& OnComplete)
+{
+	UToolsetRegistrySubsystem* ToolsetRegistrySubsystem = GEditor ? GEditor->GetEditorSubsystem<UToolsetRegistrySubsystem>() : nullptr;
+	if (!ToolsetRegistrySubsystem)
+	{
+		OnComplete(UE::ModelContextProtocol::MakeErrorResult(TEXT("ToolsetRegistry not available")));
+		return;
+	}
+
+	auto DescriptorResult = UE::ToolsetRegistry::FToolDescriptor::FromString(ToolName);
+	if (!DescriptorResult.HasValue())
+	{
+		OnComplete(UE::ModelContextProtocol::MakeErrorResult(FString::Printf(TEXT("Invalid tool descriptor: %s"), *ToolName)));
+		return;
+	}
+
+	FString ArgumentsJson;
+	if (Params.IsValid())
+	{
+		TSharedRef<TJsonWriter<>> Writer = TJsonWriterFactory<>::Create(&ArgumentsJson);
+		FJsonSerializer::Serialize(Params.ToSharedRef(), Writer);
+	}
+	else
+	{
+		ArgumentsJson = TEXT("{}");
+	}
+
+	UE_LOG(LogModelContextProtocol, Log, TEXT("Running ToolsetRegistry tool: '%s'"), *ToolName);
+
+	UE::ToolsetRegistry::FToolsetRegistry& ToolsetRegistry = ToolsetRegistrySubsystem->ToolsetRegistry;
+	ToolsetRegistry.ExecuteTool(DescriptorResult.GetValue(), ArgumentsJson).Then([OnComplete](TFuture<TValueOrError<FString, FString>> Future)
+	{
+		TValueOrError<FString, FString> Result = Future.Get();
+
+		if (Result.HasError())
+		{
+			OnComplete(UE::ModelContextProtocol::MakeErrorResult(Result.GetError()));
+		}
+		else
+		{
+			OnComplete(UE::ModelContextProtocol::MakeTextResult(Result.GetValue()));
+		}
+	});
+}
+
+void FToolsetRegistryToolAdapterManager::RegisterTools()
+{
+	DeregisterTools();
+
+	UToolsetRegistrySubsystem* ToolsetRegistrySubsystem = GEditor ? GEditor->GetEditorSubsystem<UToolsetRegistrySubsystem>() : nullptr;
+	if (!ToolsetRegistrySubsystem)
+	{
+		return;
+	}
+
+	IModelContextProtocolModule* Module = IModelContextProtocolModule::Get();
+	if (!Module)
+	{
+		return;
+	}
+
+	UE::ToolsetRegistry::FToolsetRegistry& ToolsetRegistry = ToolsetRegistrySubsystem->ToolsetRegistry;
+	const FString SchemasJson = ToolsetRegistry.GetToolsetJsonSchemas();
+
+	TArray<TSharedPtr<FJsonValue>> ToolsetArray;
+	TSharedRef<TJsonReader<>> Reader = TJsonReaderFactory<>::Create(SchemasJson);
+	if (!FJsonSerializer::Deserialize(Reader, ToolsetArray))
+	{
+		return;
+	}
+
+	for (const TSharedPtr<FJsonValue>& ToolsetValue : ToolsetArray)
+	{
+		const TSharedPtr<FJsonObject>& ToolsetObject = ToolsetValue->AsObject();
+		if (!ToolsetObject.IsValid())
+		{
+			continue;
+		}
+
+		const TArray<TSharedPtr<FJsonValue>>* ToolsField = nullptr;
+		if (!ToolsetObject->TryGetArrayField(TEXT("tools"), ToolsField) || !ToolsField)
+		{
+			continue;
+		}
+
+		for (const TSharedPtr<FJsonValue>& ToolValue : *ToolsField)
+		{
+			const TSharedPtr<FJsonObject>& ToolJsonObject = ToolValue->AsObject();
+			if (!ToolJsonObject.IsValid())
+			{
+				continue;
+			}
+
+			FString ToolName;
+			ToolJsonObject->TryGetStringField(TEXT("name"), ToolName);
+			if (ToolName.IsEmpty())
+			{
+				continue;
+			}
+
+			FString ToolDescription;
+			ToolJsonObject->TryGetStringField(TEXT("description"), ToolDescription);
+
+			TSharedPtr<FJsonObject> InputSchema;
+			const TSharedPtr<FJsonObject>* InputSchemaField = nullptr;
+			if (ToolJsonObject->TryGetObjectField(TEXT("inputSchema"), InputSchemaField) && InputSchemaField)
+			{
+				InputSchema = *InputSchemaField;
+			}
+
+			TSharedRef<FToolsetRegistryToolAdapter> Adapter = MakeShared<FToolsetRegistryToolAdapter>(ToolName, ToolDescription, InputSchema);
+			if (Module->AddTool(Adapter))
+			{
+				RegisteredAdapters.Add(Adapter);
+			}
+		}
+	}
+
+	UE_LOG(LogModelContextProtocol, Log, TEXT("Registered %d ToolsetRegistry tool adapters"), RegisteredAdapters.Num());
+}
+
+void FToolsetRegistryToolAdapterManager::DeregisterTools()
+{
+	if (RegisteredAdapters.IsEmpty())
+	{
+		return;
+	}
+
+	if (IModelContextProtocolModule* Module = IModelContextProtocolModule::Get())
+	{
+		for (const TSharedRef<FToolsetRegistryToolAdapter>& Adapter : RegisteredAdapters)
+		{
+			Module->RemoveTool(Adapter);
+		}
+	}
+	RegisteredAdapters.Reset();
+}
diff --git a/Engine/Plugins/Experimental/ModelContextProtocol/Source/ModelContextProtocolEditor/Private/ModelContextProtocolToolsetRegistryAdapter.h b/Engine/Plugins/Experimental/ModelContextProtocol/Source/ModelContextProtocolEditor/Private/ModelContextProtocolToolsetRegistryAdapter.h
new file mode 100644
--- /dev/null
+++ b/Engine/Plugins/Experimental/ModelContextProtocol/Source/ModelContextProtocolEditor/Private/ModelContextProtocolToolsetRegistryAdapter.h
@@ -0,0 +1,41 @@
+// Copyright Epic Games, Inc. All Rights Reserved.
+
+#pragma once
+
+#include "IModelContextProtocolTool.h"
+
+namespace UE::ToolsetRegistry { struct FToolDescriptor; class FToolsetRegistry; }
+
+/**
+ * Adapts a ToolsetRegistry tool to the IModelContextProtocolTool interface.
+ * Created by the Editor module to bridge ToolsetRegistry tools into the MCP server.
+ */
+struct FToolsetRegistryToolAdapter : IModelContextProtocolTool
+{
+	FToolsetRegistryToolAdapter(const FString& InToolName, const FString& InDescription, TSharedPtr<FJsonObject> InInputSchema);
+
+	virtual FString GetName() const override { return ToolName; }
+	virtual FString GetDescription() const override { return Description; }
+	virtual TSharedPtr<FJsonObject> GetInputJsonSchema() const override { return InputSchema; }
+	virtual void RunAsync(const FModelContextProtocolToolRequestId& RequestId, const TSharedPtr<FJsonObject>& Params, const FResultCallback& OnComplete) override;
+
+private:
+	FString ToolName;
+	FString Description;
+	TSharedPtr<FJsonObject> InputSchema;
+};
+
+/**
+ * Registers and manages ToolsetRegistry tool adapters with the MCP module.
+ */
+struct FToolsetRegistryToolAdapterManager
+{
+	/** Enumerates ToolsetRegistry tools and registers IModelContextProtocolTool adapters via AddTool(). */
+	void RegisterTools();
+
+	/** Removes all previously registered adapter tools. */
+	void DeregisterTools();
+
+private:
+	TArray<TSharedRef<FToolsetRegistryToolAdapter>> RegisteredAdapters;
+};
diff --git a/Engine/Plugins/Experimental/ModelContextProtocol/Source/ModelContextProtocolEditor/Public/ModelContextProtocolEditor.h b/Engine/Plugins/Experimental/ModelContextProtocol/Source/ModelContextProtocolEditor/Public/ModelContextProtocolEditor.h
--- a/Engine/Plugins/Experimental/ModelContextProtocol/Source/ModelContextProtocolEditor/Public/ModelContextProtocolEditor.h
+++ b/Engine/Plugins/Experimental/ModelContextProtocol/Source/ModelContextProtocolEditor/Public/ModelContextProtocolEditor.h
@@ -2,6 +2,8 @@
 
 #pragma once
 
+#include "ModelContextProtocolToolsetRegistryAdapter.h"
+
 #include "Modules/ModuleManager.h"
 
 class FModelContextProtocolEditorModule : public IModuleInterface
@@ -11,4 +13,10 @@ public:
 	/** IModuleInterface implementation */
 	virtual void StartupModule() override;
 	virtual void ShutdownModule() override;
+
+private:
+	FToolsetRegistryToolAdapterManager ToolsetRegistryAdapterManager;
+	FDelegateHandle PostEngineInitHandle;
+	FDelegateHandle OnRefreshToolsHandle;
+	FDelegateHandle OnToolsetRegisteredHandle;
 };
diff --git a/Engine/Plugins/Experimental/ModelContextProtocol/Source/ModelContextProtocolEditorTests/ModelContextProtocolEditorTests.Build.cs b/Engine/Plugins/Experimental/ModelContextProtocol/Source/ModelContextProtocolEditorTests/ModelContextProtocolEditorTests.Build.cs
--- a/Engine/Plugins/Experimental/ModelContextProtocol/Source/ModelContextProtocolEditorTests/ModelContextProtocolEditorTests.Build.cs
+++ b/Engine/Plugins/Experimental/ModelContextProtocol/Source/ModelContextProtocolEditorTests/ModelContextProtocolEditorTests.Build.cs
@@ -8,6 +8,9 @@ public class ModelContextProtocolEditorTests : ModuleRules
 	{
 		PCHUsage = ModuleRules.PCHUsageMode.UseExplicitOrSharedPCHs;
 
+		// Include paths for shared test utilities and mocks from ModelContextProtocolTests
+		PrivateIncludePaths.Add(System.IO.Path.Combine(PluginDirectory, "Source", "ModelContextProtocolTests", "Private"));
+
 		PublicDependencyModuleNames.AddRange(
 			new string[]
 			{
@@ -19,7 +22,14 @@ public class ModelContextProtocolEditorTests : ModuleRules
 			{
 				"CoreUObject",
 				"Engine",
-				"ModelContextProtocol"
+				"Json",
+				"JsonUtilities",
+				"HTTPServer",
+				"HTTP",
+				"ModelContextProtocol",
+				"ModelContextProtocolEngine",
+				"ModelContextProtocolEditor",
+				"ModelContextProtocolTests"
 			});
 
 		if (Target.bBuildEditor)
@@ -32,7 +42,7 @@ public class ModelContextProtocolEditorTests : ModuleRules
 					"SlateCore",
 					"AssetDefinition",
 					"Blutility",
-					"ModelContextProtocolEditor"
+					"ToolsetRegistry"
 				});
 		}
 	}
diff --git a/Engine/Plugins/Experimental/ModelContextProtocol/Source/ModelContextProtocolEditorTests/Private/Mocks/MockToolsetDefinition.cpp b/Engine/Plugins/Experimental/ModelContextProtocol/Source/ModelContextProtocolEditorTests/Private/Mocks/MockToolsetDefinition.cpp
new file mode 100644
--- /dev/null
+++ b/Engine/Plugins/Experimental/ModelContextProtocol/Source/ModelContextProtocolEditorTests/Private/Mocks/MockToolsetDefinition.cpp
@@ -0,0 +1,15 @@
+// Copyright Epic Games, Inc. All Rights Reserved.
+
+#include "Mocks/MockToolsetDefinition.h"
+
+#include UE_INLINE_GENERATED_CPP_BY_NAME(MockToolsetDefinition)
+
+FString UMockToolsetDefinition::Greet(const FString& Name)
+{
+	return FString::Printf(TEXT("Hello, %s!"), *Name);
+}
+
+int32 UMockToolsetDefinition::Add(int32 A, int32 B)
+{
+	return A + B;
+}
diff --git a/Engine/Plugins/Experimental/ModelContextProtocol/Source/ModelContextProtocolEditorTests/Private/Mocks/MockToolsetDefinition.h b/Engine/Plugins/Experimental/ModelContextProtocol/Source/ModelContextProtocolEditorTests/Private/Mocks/MockToolsetDefinition.h
new file mode 100644
--- /dev/null
+++ b/Engine/Plugins/Experimental/ModelContextProtocol/Source/ModelContextProtocolEditorTests/Private/Mocks/MockToolsetDefinition.h
@@ -0,0 +1,25 @@
+// Copyright Epic Games, Inc. All Rights Reserved.
+
+#pragma once
+
+#include "ToolsetRegistry/UToolsetRegistry.h"
+
+#include "MockToolsetDefinition.generated.h"
+
+/**
+ * Mock UToolsetDefinition subclass for testing ToolsetRegistry integration with MCP.
+ */
+UCLASS()
+class UMockToolsetDefinition : public UToolsetDefinition
+{
+	GENERATED_BODY()
+
+public:
+	/** Greets a person by name */
+	UFUNCTION()
+	static FString Greet(const FString& Name);
+
+	/** Adds two numbers together */
+	UFUNCTION()
+	static int32 Add(int32 A, int32 B);
+};
diff --git a/Engine/Plugins/Experimental/ModelContextProtocol/Source/ModelContextProtocolEditorTests/Private/Tests/ModelContextProtocolToolUtilsTests.cpp b/Engine/Plugins/Experimental/ModelContextProtocol/Source/ModelContextProtocolEditorTests/Private/Tests/ModelContextProtocolToolUtilsTests.cpp
new file mode 100644
--- /dev/null
+++ b/Engine/Plugins/Experimental/ModelContextProtocol/Source/ModelContextProtocolEditorTests/Private/Tests/ModelContextProtocolToolUtilsTests.cpp
@@ -0,0 +1,140 @@
+// Copyright Epic Games, Inc. All Rights Reserved.
+
+#include "ModelContextProtocolToolUtils.h"
+#include "ModelContextProtocolToolResults.h"
+
+#include "Engine/StaticMesh.h"
+#include "Engine/Texture2D.h"
+#include "GameFramework/Actor.h"
+#include "Sound/SoundWave.h"
+#include "UObject/UObjectIterator.h"
+#include "UObject/UnrealType.h"
+#include "Misc/AutomationTest.h"
+
+#if WITH_DEV_AUTOMATION_TESTS
+
+BEGIN_DEFINE_SPEC(FModelContextProtocolToolUtilsTests, "AI.ModelContextProtocol.ToolUtils", EAutomationTestFlags::ProductFilter | EAutomationTestFlags_ApplicationContextMask)
+END_DEFINE_SPEC(FModelContextProtocolToolUtilsTests)
+
+void FModelContextProtocolToolUtilsTests::Define()
+{
+	using namespace UE::ModelContextProtocol;
+
+	Describe("GetToolResultType", [this]()
+	{
+		It("should return None for nullptr property", [this]()
+		{
+			TSharedPtr<FJsonObject> OutSchema;
+			EModelContextProtocolToolResultType Result = GetToolResultType(nullptr, nullptr, OutSchema);
+			TestEqual("Should return None", Result, EModelContextProtocolToolResultType::None);
+		});
+
+		It("should return Text for FStrProperty", [this]()
+		{
+			// Find an FStrProperty from a known UStruct (AActor has many string properties)
+			FProperty* StringProperty = nullptr;
+			for (TFieldIterator<FStrProperty> It(AActor::StaticClass()); It; ++It)
+			{
+				StringProperty = *It;
+				break;
+			}
+			if (!TestNotNull("Should find an FStrProperty on AActor", StringProperty)) { return; }
+
+			TSharedPtr<FJsonObject> OutSchema;
+			EModelContextProtocolToolResultType Result = GetToolResultType(StringProperty, nullptr, OutSchema);
+			TestEqual("Should return Text for FStrProperty", Result, EModelContextProtocolToolResultType::Text);
+		});
+
+		It("should return Image for UTexture2D object property", [this]()
+		{
+			// Search all loaded classes for an FObjectProperty with PropertyClass == UTexture2D
+			FObjectProperty* TextureProperty = nullptr;
+			for (TObjectIterator<UClass> ClassIt; ClassIt; ++ClassIt)
+			{
+				for (TFieldIterator<FObjectProperty> PropIt(*ClassIt, EFieldIteratorFlags::ExcludeSuper); PropIt; ++PropIt)
+				{
+					if (PropIt->PropertyClass == UTexture2D::StaticClass())
+					{
+						TextureProperty = *PropIt;
+						break;
+					}
+				}
+				if (TextureProperty) { break; }
+			}
+
+			if (!TestNotNull("Should find a UTexture2D object property on any loaded class", TextureProperty)) { return; }
+
+			TSharedPtr<FJsonObject> OutSchema;
+			EModelContextProtocolToolResultType Result = GetToolResultType(TextureProperty, nullptr, OutSchema);
+			TestEqual("Should return Image for UTexture2D property", Result, EModelContextProtocolToolResultType::Image);
+		});
+
+		It("should return StructuredContent for struct properties that are not text/image/audio", [this]()
+		{
+			// Find a property that is not string, not UTexture2D, not USoundWave
+			FProperty* StructProperty = nullptr;
+			for (TFieldIterator<FStructProperty> It(AActor::StaticClass()); It; ++It)
+			{
+				StructProperty = *It;
+				break;
+			}
+			if (!TestNotNull("Should find a struct property on AActor", StructProperty)) { return; }
+
+			TSharedPtr<FJsonObject> OutSchema;
+			EModelContextProtocolToolResultType Result = GetToolResultType(StructProperty, nullptr, OutSchema);
+			TestEqual("Should return StructuredContent for struct property", Result, EModelContextProtocolToolResultType::StructuredContent);
+		});
+
+		It("should return Text for FSoftObjectProperty", [this]()
+		{
+			FProperty* SoftObjectProperty = nullptr;
+			for (TObjectIterator<UClass> ClassIt; ClassIt; ++ClassIt)
+			{
+				for (TFieldIterator<FSoftObjectProperty> PropIt(*ClassIt, EFieldIteratorFlags::ExcludeSuper); PropIt; ++PropIt)
+				{
+					SoftObjectProperty = *PropIt;
+					break;
+				}
+				if (SoftObjectProperty) { break; }
+			}
+
+			if (!TestNotNull("Should find a FSoftObjectProperty on any loaded class", SoftObjectProperty)) { return; }
+
+			TSharedPtr<FJsonObject> OutSchema;
+			EModelContextProtocolToolResultType Result = GetToolResultType(SoftObjectProperty, nullptr, OutSchema);
+			TestEqual("Should return Text for FSoftObjectProperty", Result, EModelContextProtocolToolResultType::Text);
+		});
+
+		It("should populate OutSchema for StructuredContent results", [this]()
+		{
+			FProperty* StructProperty = nullptr;
+			for (TFieldIterator<FStructProperty> It(AActor::StaticClass()); It; ++It)
+			{
+				StructProperty = *It;
+				break;
+			}
+			if (!TestNotNull("Should find a struct property", StructProperty)) { return; }
+
+			TSharedPtr<FJsonObject> OutSchema;
+			GetToolResultType(StructProperty, nullptr, OutSchema);
+			TestTrue("OutSchema should be populated for StructuredContent", OutSchema.IsValid());
+		});
+
+		It("should not populate OutSchema for Text results", [this]()
+		{
+			FProperty* StringProperty = nullptr;
+			for (TFieldIterator<FStrProperty> It(AActor::StaticClass()); It; ++It)
+			{
+				StringProperty = *It;
+				break;
+			}
+			if (!TestNotNull("Should find a string property", StringProperty)) { return; }
+
+			TSharedPtr<FJsonObject> OutSchema;
+			GetToolResultType(StringProperty, nullptr, OutSchema);
+			TestFalse("OutSchema should not be populated for Text", OutSchema.IsValid());
+		});
+	});
+}
+
+#endif // WITH_DEV_AUTOMATION_TESTS
diff --git a/Engine/Plugins/Experimental/ModelContextProtocol/Source/ModelContextProtocolEditorTests/Private/Tests/ModelContextProtocolToolsetRegistryTests.cpp b/Engine/Plugins/Experimental/ModelContextProtocol/Source/ModelContextProtocolEditorTests/Private/Tests/ModelContextProtocolToolsetRegistryTests.cpp
new file mode 100644
--- /dev/null
+++ b/Engine/Plugins/Experimental/ModelContextProtocol/Source/ModelContextProtocolEditorTests/Private/Tests/ModelContextProtocolToolsetRegistryTests.cpp
@@ -0,0 +1,403 @@
+// Copyright Epic Games, Inc. All Rights Reserved.
+
+#include "IModelContextProtocolModule.h"
+#include "ModelContextProtocol.h"
+#include "ModelContextProtocolTestUtilities.h"
+#include "Mocks/MockModelContextProtocolTool.h"
+
+#include "HttpModule.h"
+#include "Interfaces/IHttpRequest.h"
+#include "Interfaces/IHttpResponse.h"
+#include "Misc/AutomationTest.h"
+
+#if WITH_EDITOR
+#include "Editor.h"
+#include "ToolsetRegistry/BlueprintLibraryToolset.h"
+#include "ToolsetRegistry/ToolsetRegistry.h"
+#include "ToolsetRegistry/ToolsetRegistrySubsystem.h"
+#include "Mocks/MockToolsetDefinition.h"
+#endif
+
+#if WITH_DEV_AUTOMATION_TESTS && WITH_EDITOR
+
+// Qualified tool names follow the pattern: {ModuleName}.{ClassName}.{FunctionName}
+static const FString MockToolsetName = TEXT("ModelContextProtocolEditorTests.MockToolsetDefinition");
+static const FString GreetToolName = MockToolsetName + TEXT(".Greet");
+static const FString AddToolName = MockToolsetName + TEXT(".Add");
+
+BEGIN_DEFINE_SPEC(FModelContextProtocolToolsetRegistryTests, "AI.ModelContextProtocol.ToolsetRegistry", EAutomationTestFlags::ProductFilter | EAutomationTestFlags_ApplicationContextMask)
+	FString SessionId;
+	TSharedPtr<UE::ToolsetRegistry::FBlueprintLibraryToolset> MockToolset;
+
+	void SendJsonRpcRequest(const TSharedRef<FJsonObject>& JsonRpcBody, const TFunction<void(FHttpResponsePtr, bool)>& OnComplete, const FString& InSessionId = FString(), TOptional<float> ActivityTimeout = {})
+	{
+		FHttpModule& HttpModule = FHttpModule::Get();
+		TSharedRef<IHttpRequest> HttpRequest = HttpModule.CreateRequest();
+		HttpRequest->SetURL(UE::ModelContextProtocol::Tests::GetTestBaseUrl());
+		HttpRequest->SetVerb(TEXT("POST"));
+		HttpRequest->SetHeader(TEXT("Content-Type"), TEXT("application/json"));
+		if (!InSessionId.IsEmpty())
+		{
+			HttpRequest->SetHeader(TEXT("Mcp-Session-Id"), InSessionId);
+		}
+		if (ActivityTimeout.IsSet())
+		{
+			HttpRequest->SetActivityTimeout(ActivityTimeout.GetValue());
+		}
+		HttpRequest->SetContentAsString(UE::ModelContextProtocol::Tests::JsonObjectToString(JsonRpcBody));
+		HttpRequest->OnProcessRequestComplete().BindLambda([OnComplete](FHttpRequestPtr, FHttpResponsePtr Response, bool bConnectedSuccessfully)
+		{
+			OnComplete(Response, bConnectedSuccessfully);
+		});
+		HttpRequest->ProcessRequest();
+	}
+
+	TSharedPtr<FJsonObject> ParseJsonResponse(FHttpResponsePtr Response)
+	{
+		if (!Response.IsValid())
+		{
+			return nullptr;
+		}
+		TSharedPtr<FJsonObject> JsonObject;
+		TSharedRef<TJsonReader<>> Reader = TJsonReaderFactory<>::Create(Response->GetContentAsString());
+		FJsonSerializer::Deserialize(Reader, JsonObject);
+		return JsonObject;
+	}
+
+	void InitializeSession(const FDoneDelegate& Done)
+	{
+		using namespace UE::ModelContextProtocol::Tests;
+
+		TSharedRef<FJsonObject> InitParams = MakeShared<FJsonObject>();
+		InitParams->SetStringField(TEXT("protocolVersion"), UE::ModelContextProtocol::ProtocolVersion);
+		TSharedRef<FJsonObject> ClientInfo = MakeShared<FJsonObject>();
+		ClientInfo->SetStringField(TEXT("name"), TEXT("TestClient"));
+		ClientInfo->SetStringField(TEXT("version"), TEXT("1.0"));
+		InitParams->SetObjectField(TEXT("clientInfo"), ClientInfo);
+		InitParams->SetObjectField(TEXT("capabilities"), MakeShared<FJsonObject>());
+
+		TSharedRef<FJsonObject> InitRequest = MakeJsonRpcRequest(TEXT("initialize"), MakeShared<FJsonValueNumber>(1), InitParams);
+		SendJsonRpcRequest(InitRequest, [this, Done](FHttpResponsePtr Response, bool bConnected)
+		{
+			if (Response.IsValid())
+			{
+				SessionId = Response->GetHeader(TEXT("Mcp-Session-Id"));
+
+				TSharedRef<FJsonObject> InitializedNotification = UE::ModelContextProtocol::Tests::MakeJsonRpcNotification(TEXT("notifications/initialized"));
+				SendJsonRpcRequest(InitializedNotification, [Done](FHttpResponsePtr, bool)
+				{
+					Done.Execute();
+				}, SessionId);
+			}
+			else
+			{
+				Done.Execute();
+			}
+		});
+	}
+
+	void CleanupSession(const FDoneDelegate& Done)
+	{
+		if (SessionId.IsEmpty())
+		{
+			Done.Execute();
+			return;
+		}
+
+		FHttpModule& HttpModule = FHttpModule::Get();
+		TSharedRef<IHttpRequest> DeleteRequest = HttpModule.CreateRequest();
+		DeleteRequest->SetURL(UE::ModelContextProtocol::Tests::GetTestBaseUrl());
+		DeleteRequest->SetVerb(TEXT("DELETE"));
+		DeleteRequest->SetHeader(TEXT("Mcp-Session-Id"), SessionId);
+		DeleteRequest->OnProcessRequestComplete().BindLambda([this, Done](FHttpRequestPtr, FHttpResponsePtr, bool)
+		{
+			SessionId.Reset();
+			Done.Execute();
+		});
+		DeleteRequest->ProcessRequest();
+	}
+END_DEFINE_SPEC(FModelContextProtocolToolsetRegistryTests)
+
+void FModelContextProtocolToolsetRegistryTests::Define()
+{
+	using namespace UE::ModelContextProtocol::Tests;
+
+	BeforeEach([this]()
+	{
+		SessionId.Reset();
+
+		// Register mock toolset with ToolsetRegistry
+		if (UToolsetRegistrySubsystem* ToolsetRegistrySubsystem = GEditor ? GEditor->GetEditorSubsystem<UToolsetRegistrySubsystem>() : nullptr)
+		{
+			MockToolset = MakeShared<UE::ToolsetRegistry::FBlueprintLibraryToolset>(UMockToolsetDefinition::StaticClass());
+			ToolsetRegistrySubsystem->ToolsetRegistry.RegisterToolset(MockToolset);
+		}
+
+		// Trigger adapter re-registration so the mock toolset's tools are available via IModelContextProtocolTool
+		IModelContextProtocolModule::GetChecked().RefreshTools();
+	});
+
+	AfterEach([this]()
+	{
+		if (MockToolset.IsValid())
+		{
+			if (UToolsetRegistrySubsystem* ToolsetRegistrySubsystem = GEditor ? GEditor->GetEditorSubsystem<UToolsetRegistrySubsystem>() : nullptr)
+			{
+				ToolsetRegistrySubsystem->ToolsetRegistry.UnregisterToolset(MockToolset);
+			}
+			MockToolset.Reset();
+		}
+	});
+
+	Describe("ToolsetRegistry Tool Listing", [this]()
+	{
+		LatentIt("should include ToolsetRegistry tools in tools/list", [this](const FDoneDelegate& Done)
+		{
+			InitializeSession(FDoneDelegate::CreateLambda([this, Done]()
+			{
+				if (SessionId.IsEmpty())
+				{
+					Done.Execute();
+					return;
+				}
+
+				TSharedRef<FJsonObject> ListRequest = MakeJsonRpcRequest(TEXT("tools/list"), MakeShared<FJsonValueNumber>(2));
+				SendJsonRpcRequest(ListRequest, [this, Done](FHttpResponsePtr Response, bool bConnected)
+				{
+					if (TestTrue("Should connect", bConnected) && TestTrue("Should have response", Response.IsValid()))
+					{
+						TSharedPtr<FJsonObject> JsonResponse = ParseJsonResponse(Response);
+						if (TestTrue("Should parse response", JsonResponse.IsValid()))
+						{
+							const TSharedPtr<FJsonObject>* ResultObject;
+							if (TestTrue("Should have result", JsonResponse->TryGetObjectField(TEXT("result"), ResultObject)))
+							{
+								const TArray<TSharedPtr<FJsonValue>>* ToolsArray;
+								if (TestTrue("Should have tools array", (*ResultObject)->TryGetArrayField(TEXT("tools"), ToolsArray)))
+								{
+									bool bFoundGreet = false;
+									bool bFoundAdd = false;
+									for (const TSharedPtr<FJsonValue>& ToolValue : *ToolsArray)
+									{
+										const TSharedPtr<FJsonObject>* ToolObject;
+										if (ToolValue->TryGetObject(ToolObject))
+										{
+											FString ToolName;
+											(*ToolObject)->TryGetStringField(TEXT("name"), ToolName);
+											if (ToolName == GreetToolName)
+											{
+												bFoundGreet = true;
+												FString ToolDescription;
+												(*ToolObject)->TryGetStringField(TEXT("description"), ToolDescription);
+												TestFalse("Greet tool should have a description", ToolDescription.IsEmpty());
+												const TSharedPtr<FJsonObject>* InputSchemaField;
+												TestTrue("Greet tool should have inputSchema", (*ToolObject)->TryGetObjectField(TEXT("inputSchema"), InputSchemaField));
+											}
+											else if (ToolName == AddToolName)
+											{
+												bFoundAdd = true;
+												const TSharedPtr<FJsonObject>* InputSchemaField;
+												TestTrue("Add tool should have inputSchema", (*ToolObject)->TryGetObjectField(TEXT("inputSchema"), InputSchemaField));
+											}
+										}
+									}
+									TestTrue("Should find Greet tool from ToolsetRegistry", bFoundGreet);
+									TestTrue("Should find Add tool from ToolsetRegistry", bFoundAdd);
+								}
+							}
+						}
+					}
+					CleanupSession(Done);
+				}, SessionId);
+			}));
+		});
+	});
+
+	Describe("ToolsetRegistry Tool Execution", [this]()
+	{
+		LatentIt("should execute ToolsetRegistry tool via tools/call", [this](const FDoneDelegate& Done)
+		{
+			InitializeSession(FDoneDelegate::CreateLambda([this, Done]()
+			{
+				if (SessionId.IsEmpty())
+				{
+					Done.Execute();
+					return;
+				}
+
+				TSharedRef<FJsonObject> ToolArguments = MakeShared<FJsonObject>();
+				ToolArguments->SetStringField(TEXT("Name"), TEXT("World"));
+
+				TSharedRef<FJsonObject> CallParams = MakeShared<FJsonObject>();
+				CallParams->SetStringField(TEXT("name"), GreetToolName);
+				CallParams->SetObjectField(TEXT("arguments"), ToolArguments);
+
+				TSharedRef<FJsonObject> CallRequest = MakeJsonRpcRequest(TEXT("tools/call"), MakeShared<FJsonValueNumber>(3), CallParams);
+				// tools/call returns an SSE stream; use a short activity timeout since the tool completes synchronously
+				SendJsonRpcRequest(CallRequest, [this, Done](FHttpResponsePtr Response, bool bConnected)
+				{
+					if (TestTrue("Should have response", Response.IsValid()))
+					{
+						TSharedPtr<FJsonObject> JsonResponse = ParseSseJsonResponse(Response);
+						if (JsonResponse.IsValid())
+						{
+							const TSharedPtr<FJsonObject>* ResultObject;
+							if (TestTrue("Should have result", JsonResponse->TryGetObjectField(TEXT("result"), ResultObject)))
+							{
+								const TArray<TSharedPtr<FJsonValue>>* ContentArray;
+								if (TestTrue("Should have content", (*ResultObject)->TryGetArrayField(TEXT("content"), ContentArray))
+									&& TestTrue("Content should have at least 1 element", ContentArray->Num() >= 1))
+								{
+									const TSharedPtr<FJsonObject>* ContentObject;
+									if ((*ContentArray)[0]->TryGetObject(ContentObject))
+									{
+										FString Text;
+										if ((*ContentObject)->TryGetStringField(TEXT("text"), Text))
+										{
+											TestTrue("Tool result should contain greeting", Text.Contains(TEXT("Hello, World!")));
+										}
+									}
+								}
+							}
+						}
+						else
+						{
+							int32 ResponseCode = Response->GetResponseCode();
+							TestTrue("Should return success status for tool call", ResponseCode >= 200 && ResponseCode < 300);
+						}
+					}
+					CleanupSession(Done);
+				}, SessionId, 2.0f);
+			}));
+		});
+
+		LatentIt("should execute ToolsetRegistry tool with numeric result via tools/call", [this](const FDoneDelegate& Done)
+		{
+			InitializeSession(FDoneDelegate::CreateLambda([this, Done]()
+			{
+				if (SessionId.IsEmpty())
+				{
+					Done.Execute();
+					return;
+				}
+
+				TSharedRef<FJsonObject> ToolArguments = MakeShared<FJsonObject>();
+				ToolArguments->SetNumberField(TEXT("A"), 3);
+				ToolArguments->SetNumberField(TEXT("B"), 7);
+
+				TSharedRef<FJsonObject> CallParams = MakeShared<FJsonObject>();
+				CallParams->SetStringField(TEXT("name"), AddToolName);
+				CallParams->SetObjectField(TEXT("arguments"), ToolArguments);
+
+				TSharedRef<FJsonObject> CallRequest = MakeJsonRpcRequest(TEXT("tools/call"), MakeShared<FJsonValueNumber>(4), CallParams);
+				// tools/call returns an SSE stream; use a short activity timeout since the tool completes synchronously
+				SendJsonRpcRequest(CallRequest, [this, Done](FHttpResponsePtr Response, bool bConnected)
+				{
+					if (TestTrue("Should have response", Response.IsValid()))
+					{
+						TSharedPtr<FJsonObject> JsonResponse = ParseSseJsonResponse(Response);
+						if (JsonResponse.IsValid())
+						{
+							const TSharedPtr<FJsonObject>* ResultObject;
+							if (TestTrue("Should have result", JsonResponse->TryGetObjectField(TEXT("result"), ResultObject)))
+							{
+								const TArray<TSharedPtr<FJsonValue>>* ContentArray;
+								if (TestTrue("Should have content", (*ResultObject)->TryGetArrayField(TEXT("content"), ContentArray))
+									&& TestTrue("Content should have at least 1 element", ContentArray->Num() >= 1))
+								{
+									const TSharedPtr<FJsonObject>* ContentObject;
+									if ((*ContentArray)[0]->TryGetObject(ContentObject))
+									{
+										FString Text;
+										if ((*ContentObject)->TryGetStringField(TEXT("text"), Text))
+										{
+											TestTrue("Result should contain the sum", Text.Contains(TEXT("10")));
+										}
+									}
+								}
+							}
+						}
+						else
+						{
+							int32 ResponseCode = Response->GetResponseCode();
+							TestTrue("Should return success status for tool call", ResponseCode >= 200 && ResponseCode < 300);
+						}
+					}
+					CleanupSession(Done);
+				}, SessionId, 2.0f);
+			}));
+		});
+	});
+
+	Describe("ToolsetRegistry Tool Deduplication", [this]()
+	{
+		LatentIt("should deduplicate when MCP and ToolsetRegistry have same tool name", [this](const FDoneDelegate& Done)
+		{
+			InitializeSession(FDoneDelegate::CreateLambda([this, Done]()
+			{
+				if (SessionId.IsEmpty())
+				{
+					Done.Execute();
+					return;
+				}
+
+				// Register a deprecated MCP tool with the same qualified name as a ToolsetRegistry tool.
+				// AddTool() rejects duplicates at registration time, so the ToolsetRegistry adapter's
+				// tool (already registered via RefreshTools in BeforeEach) takes precedence.
+				TSharedPtr<FMockModelContextProtocolTool> MockMcpTool = MakeShared<FMockModelContextProtocolTool>();
+				MockMcpTool->Name = GreetToolName;
+				MockMcpTool->Description = TEXT("Deprecated MCP version of Greet tool");
+				MockMcpTool->InputSchema = FMockModelContextProtocolTool::MakeTestInputSchema(
+					{{TEXT("Name"), TEXT("string")}},
+					{TEXT("Name")});
+
+				IModelContextProtocolModule& Module = IModelContextProtocolModule::GetChecked();
+				AddExpectedError(TEXT("is already registered"), EAutomationExpectedErrorFlags::Contains);
+				TestFalse("Duplicate tool should be rejected by AddTool", Module.AddTool(MockMcpTool.ToSharedRef()));
+
+				TSharedRef<FJsonObject> ListRequest = MakeJsonRpcRequest(TEXT("tools/list"), MakeShared<FJsonValueNumber>(5));
+				SendJsonRpcRequest(ListRequest, [this, Done, MockMcpTool](FHttpResponsePtr Response, bool bConnected)
+				{
+					if (TestTrue("Should connect", bConnected) && TestTrue("Should have response", Response.IsValid()))
+					{
+						TSharedPtr<FJsonObject> JsonResponse = ParseJsonResponse(Response);
+						if (TestTrue("Should parse response", JsonResponse.IsValid()))
+						{
+							const TSharedPtr<FJsonObject>* ResultObject;
+							if (TestTrue("Should have result", JsonResponse->TryGetObjectField(TEXT("result"), ResultObject)))
+							{
+								const TArray<TSharedPtr<FJsonValue>>* ToolsArray;
+								if (TestTrue("Should have tools array", (*ResultObject)->TryGetArrayField(TEXT("tools"), ToolsArray)))
+								{
+									int32 GreetCount = 0;
+									FString GreetDescription;
+									for (const TSharedPtr<FJsonValue>& ToolValue : *ToolsArray)
+									{
+										const TSharedPtr<FJsonObject>* ToolObject;
+										if (ToolValue->TryGetObject(ToolObject))
+										{
+											FString ToolName;
+											(*ToolObject)->TryGetStringField(TEXT("name"), ToolName);
+											if (ToolName == GreetToolName)
+											{
+												GreetCount++;
+												(*ToolObject)->TryGetStringField(TEXT("description"), GreetDescription);
+											}
+										}
+									}
+									TestEqual("Should have exactly one Greet tool (deduplicated)", GreetCount, 1);
+									TestFalse("ToolsetRegistry tool should take precedence (not the deprecated MCP version)",
+										GreetDescription == TEXT("Deprecated MCP version of Greet tool"));
+								}
+							}
+						}
+					}
+
+					CleanupSession(Done);
+				}, SessionId);
+			}));
+		});
+	});
+}
+
+#endif // WITH_DEV_AUTOMATION_TESTS && WITH_EDITOR
diff --git a/Engine/Plugins/Experimental/ModelContextProtocol/Source/ModelContextProtocolEngine/ModelContextProtocolEngine.Build.cs b/Engine/Plugins/Experimental/ModelContextProtocol/Source/ModelContextProtocolEngine/ModelContextProtocolEngine.Build.cs
new file mode 100644
--- /dev/null
+++ b/Engine/Plugins/Experimental/ModelContextProtocol/Source/ModelContextProtocolEngine/ModelContextProtocolEngine.Build.cs
@@ -0,0 +1,39 @@
+// Copyright Epic Games, Inc. All Rights Reserved.
+
+using UnrealBuildTool;
+
+public class ModelContextProtocolEngine : ModuleRules
+{
+	public ModelContextProtocolEngine(ReadOnlyTargetRules Target) : base(Target)
+	{
+		PCHUsage = ModuleRules.PCHUsageMode.UseExplicitOrSharedPCHs;
+
+		PublicDependencyModuleNames.AddRange(
+			new string[]
+			{
+				"Core",
+				"ModelContextProtocol"
+			});
+
+		PrivateDependencyModuleNames.AddRange(
+			new string[]
+			{
+				"CoreUObject",
+				"Engine",
+				"Json",
+				"JsonUtilities",
+				"DeveloperSettings"
+			});
+
+		if (Target.bBuildEditor)
+		{
+			PrivateDependencyModuleNames.AddRange(
+				new string[]
+				{
+					"BlueprintGraph",
+					"JsonUtilitiesEditor",
+					"UnrealEd"
+				});
+		}
+	}
+}
diff --git a/Engine/Plugins/Experimental/ModelContextProtocol/Source/ModelContextProtocolEngine/Private/ModelContextProtocolEngineMetaData.cpp b/Engine/Plugins/Experimental/ModelContextProtocol/Source/ModelContextProtocolEngine/Private/ModelContextProtocolEngineMetaData.cpp
new file mode 100644
--- /dev/null
+++ b/Engine/Plugins/Experimental/ModelContextProtocol/Source/ModelContextProtocolEngine/Private/ModelContextProtocolEngineMetaData.cpp
@@ -0,0 +1,200 @@
+// Copyright Epic Games, Inc. All Rights Reserved.
+
+#include "ModelContextProtocolEngineMetaData.h"
+
+#include "JsonSchema/JsonSchemaEditorMetadata.h"
+#include "JsonSchema/JsonSchemaPropertyFilter.h"
+#include "Serialization/JsonReader.h"
+#include "Serialization/JsonSerializer.h"
+
+#if WITH_EDITORONLY_DATA
+#include "EdGraphSchema_K2.h"
+#include "JsonSchema/JsonSchemaGeneratorEditor.h"
+#endif // WITH_EDITORONLY_DATA
+
+namespace UE::ModelContextProtocol
+{
+
+#if WITH_EDITORONLY_DATA
+FModelContextProtocolFunctionMetaData CollectFunctionMetaData(const UFunction* Function, const FJsonSchemaPropertyFilter& PropertyFilter)
+{
+	FModelContextProtocolFunctionMetaData Result;
+
+	if (!Function)
+	{
+		return Result;
+	}
+
+	FJsonSchemaEditorMetadata EditorMetadata = FJsonSchemaGeneratorEditor::UStructToJsonSchemaMetadata(Function, PropertyFilter, nullptr);
+
+	if (EditorMetadata.RootStructMetadata.IsValid())
+	{
+		FString DescriptionString;
+		if (EditorMetadata.RootStructMetadata->TryGetStringField(TEXT("description"), DescriptionString) && !DescriptionString.IsEmpty())
+		{
+			Result.Description = FText::FromString(DescriptionString);
+		}
+	}
+
+	if (Function->HasMetaData(FBlueprintMetadata::MD_WorldContext))
+	{
+		Result.WorldContext = Function->GetMetaData(FBlueprintMetadata::MD_WorldContext);
+	}
+
+	for (const auto& [PropertyPath, PropertyMetadataJsonObject] : EditorMetadata.GetPropertyMemberPathToPropertyMetadataMap())
+	{
+		if (!PropertyMetadataJsonObject.IsValid())
+		{
+			continue;
+		}
+
+		FModelContextProtocolPropertyMetaData PropertyMetaData;
+
+		FString DescriptionString;
+		if (PropertyMetadataJsonObject->TryGetStringField(TEXT("description"), DescriptionString) && !DescriptionString.IsEmpty())
+		{
+			PropertyMetaData.Description = FText::FromString(DescriptionString);
+		}
+
+		if (PropertyMetadataJsonObject->HasField(TEXT("default")))
+		{
+			const TSharedPtr<FJsonValue>& DefaultValue = PropertyMetadataJsonObject->Values.FindChecked(TEXT("default"));
+			if (DefaultValue.IsValid())
+			{
+				PropertyMetaData.DefaultValueJsonType = static_cast<uint8>(DefaultValue->Type);
+				switch (DefaultValue->Type)
+				{
+				case EJson::String:
+					PropertyMetaData.DefaultValue = DefaultValue->AsString();
+					break;
+				case EJson::Number:
+					PropertyMetaData.DefaultValue = FString::SanitizeFloat(DefaultValue->AsNumber());
+					break;
+				case EJson::Boolean:
+					PropertyMetaData.DefaultValue = DefaultValue->AsBool() ? TEXT("true") : TEXT("false");
+					break;
+				case EJson::Null:
+					PropertyMetaData.DefaultValue = TEXT("None");
+					break;
+				case EJson::Object:
+				case EJson::Array:
+					{
+						FString JsonString;
+						TSharedRef<TJsonWriter<>> Writer = TJsonWriterFactory<>::Create(&JsonString);
+						FJsonSerializer::Serialize(DefaultValue, TEXT(""), Writer);
+						PropertyMetaData.DefaultValue = JsonString;
+					}
+					break;
+				default:
+					break;
+				}
+			}
+		}
+
+		double MinValue;
+		if (PropertyMetadataJsonObject->TryGetNumberField(TEXT("minimum"), MinValue))
+		{
+			PropertyMetaData.ClampMin = MinValue;
+		}
+
+		double MaxValue;
+		if (PropertyMetadataJsonObject->TryGetNumberField(TEXT("maximum"), MaxValue))
+		{
+			PropertyMetaData.ClampMax = MaxValue;
+		}
+
+		Result.PropertyMetaData.Add(PropertyPath, MoveTemp(PropertyMetaData));
+	}
+
+	return Result;
+}
+#endif // WITH_EDITORONLY_DATA
+
+FJsonSchemaEditorMetadata ConvertToCachedEditorMetadata(const FModelContextProtocolFunctionMetaData& FunctionMetaData)
+{
+	FJsonSchemaEditorMetadata EditorMetadata;
+
+	// Set root struct description
+	if (FunctionMetaData.Description.IsSet())
+	{
+		EditorMetadata.RootStructMetadata = MakeShared<FJsonObject>();
+		EditorMetadata.RootStructMetadata->SetStringField(TEXT("description"), FunctionMetaData.Description->ToString());
+	}
+
+	// Convert per-property metadata
+	FJsonSchemaEditorMetadata::FJsonSchemaPropertyMemberPathToPropertyMetadataMap& PropertyMap = EditorMetadata.GetPropertyMemberPathToPropertyMetadataMap();
+	for (const auto& [PropertyPath, PropertyMetaData] : FunctionMetaData.PropertyMetaData)
+	{
+		TSharedPtr<FJsonObject> PropertyJsonObject = MakeShared<FJsonObject>();
+
+		if (PropertyMetaData.Description.IsSet())
+		{
+			PropertyJsonObject->SetStringField(TEXT("description"), PropertyMetaData.Description->ToString());
+		}
+
+		if (PropertyMetaData.DefaultValue.IsSet())
+		{
+			const FString& DefaultString = *PropertyMetaData.DefaultValue;
+			const EJson OriginalType = PropertyMetaData.DefaultValueJsonType.IsSet()
+				? static_cast<EJson>(*PropertyMetaData.DefaultValueJsonType)
+				: EJson::String;
+
+			switch (OriginalType)
+			{
+			case EJson::Boolean:
+				PropertyJsonObject->SetBoolField(TEXT("default"), DefaultString == TEXT("true"));
+				break;
+			case EJson::Null:
+				PropertyJsonObject->SetField(TEXT("default"), MakeShared<FJsonValueNull>());
+				break;
+			case EJson::Number:
+				{
+					double NumericValue;
+					if (LexTryParseString(NumericValue, *DefaultString))
+					{
+						PropertyJsonObject->SetNumberField(TEXT("default"), NumericValue);
+					}
+					else
+					{
+						PropertyJsonObject->SetStringField(TEXT("default"), DefaultString);
+					}
+				}
+				break;
+			case EJson::Object:
+			case EJson::Array:
+				{
+					TSharedPtr<FJsonValue> ParsedValue;
+					TSharedRef<TJsonReader<>> Reader = TJsonReaderFactory<>::Create(DefaultString);
+					if (FJsonSerializer::Deserialize(Reader, ParsedValue) && ParsedValue.IsValid())
+					{
+						PropertyJsonObject->SetField(TEXT("default"), ParsedValue);
+					}
+					else
+					{
+						PropertyJsonObject->SetStringField(TEXT("default"), DefaultString);
+					}
+				}
+				break;
+			default:
+				PropertyJsonObject->SetStringField(TEXT("default"), DefaultString);
+				break;
+			}
+		}
+
+		if (PropertyMetaData.ClampMin.IsSet())
+		{
+			PropertyJsonObject->SetNumberField(TEXT("minimum"), *PropertyMetaData.ClampMin);
+		}
+
+		if (PropertyMetaData.ClampMax.IsSet())
+		{
+			PropertyJsonObject->SetNumberField(TEXT("maximum"), *PropertyMetaData.ClampMax);
+		}
+
+		PropertyMap.Add(PropertyPath, PropertyJsonObject);
+	}
+
+	return EditorMetadata;
+}
+
+}
diff --git a/Engine/Plugins/Experimental/ModelContextProtocol/Source/ModelContextProtocolEngine/Private/ModelContextProtocolEngineModule.cpp b/Engine/Plugins/Experimental/ModelContextProtocol/Source/ModelContextProtocolEngine/Private/ModelContextProtocolEngineModule.cpp
new file mode 100644
--- /dev/null
+++ b/Engine/Plugins/Experimental/ModelContextProtocol/Source/ModelContextProtocolEngine/Private/ModelContextProtocolEngineModule.cpp
@@ -0,0 +1,9 @@
+// Copyright Epic Games, Inc. All Rights Reserved.
+
+#include "Modules/ModuleManager.h"
+
+class FModelContextProtocolEngineModule : public IModuleInterface
+{
+};
+
+IMPLEMENT_MODULE(FModelContextProtocolEngineModule, ModelContextProtocolEngine);
diff --git a/Engine/Plugins/Experimental/ModelContextProtocol/Source/ModelContextProtocolEngine/Private/ModelContextProtocolEngineToolResults.cpp b/Engine/Plugins/Experimental/ModelContextProtocol/Source/ModelContextProtocolEngine/Private/ModelContextProtocolEngineToolResults.cpp
new file mode 100644
--- /dev/null
+++ b/Engine/Plugins/Experimental/ModelContextProtocol/Source/ModelContextProtocolEngine/Private/ModelContextProtocolEngineToolResults.cpp
@@ -0,0 +1,90 @@
+// Copyright Epic Games, Inc. All Rights Reserved.
+
+#include "ModelContextProtocolEngineToolResults.h"
+#include "ModelContextProtocol.h"
+
+#include "Engine/Texture2D.h"
+#include "ImageUtils.h"
+#include "Misc/ScopeExit.h"
+#include "TextureResource.h"
+#include "Sound/SoundWave.h"
+
+FModelContextProtocolToolResult UE::ModelContextProtocol::MakeImageResult(const UTexture2D* Texture, const TCHAR* ToFormatExtension, EModelContextProtocolAudience Audience)
+{
+	if (!Texture)
+	{
+		return FModelContextProtocolToolResult();
+	}
+
+	if (const FTexturePlatformData* TexturePlatformData = Texture->GetPlatformData();
+		ensure(TexturePlatformData) && ensure(!TexturePlatformData->Mips.IsEmpty()) && ensure(TexturePlatformData->PixelFormat == PF_B8G8R8A8))
+	{
+		const FTexture2DMipMap& Mip0 = TexturePlatformData->Mips[0];
+		if (Mip0.SizeX > 0 && Mip0.SizeY > 0)
+		{
+			const void* Mip0Data = Mip0.BulkData.LockReadOnly();
+			if (Mip0Data)
+			{
+				ON_SCOPE_EXIT{ Mip0.BulkData.Unlock(); };
+				return MakeImageResult(FImageView(static_cast<const FColor*>(Mip0Data), Mip0.SizeX, Mip0.SizeY), ToFormatExtension, Audience);
+			}
+		}
+	}
+
+	return FModelContextProtocolToolResult();
+}
+
+FModelContextProtocolToolResult UE::ModelContextProtocol::MakeImageResult(const FImageView& Image, const TCHAR* ToFormatExtension, EModelContextProtocolAudience Audience)
+{
+	if (!ensure(Image.IsImageInfoValid()))
+	{
+		return FModelContextProtocolToolResult();
+	}
+
+	TArray64<uint8> CompressedData;
+	FImageUtils::CompressImage(CompressedData, ToFormatExtension, Image);
+
+	FStringView Format(ToFormatExtension);
+	if (Format.StartsWith(TEXT('.')))
+	{
+		Format.RightChopInline(1);
+	}
+
+	return MakeImageResult(FString(TEXT("image/")) + Format, CompressedData, Audience);
+}
+
+FModelContextProtocolToolResult UE::ModelContextProtocol::MakeAudioResult(USoundWave* SoundWave)
+{
+	if (!SoundWave)
+	{
+		return FModelContextProtocolToolResult();
+	}
+
+#if WITH_EDITOR
+	if (UE::ModelContextProtocol::bAudioResultOggFormat)
+	{
+		if (FByteBulkData* BulkData = SoundWave->GetCompressedData("OGG"))
+		{
+			TArray<uint8> RawData;
+			RawData.Append(BulkData->GetCopyAsBuffer(BulkData->GetElementCount(), true).GetView());
+
+			return MakeAudioResult(TEXT("audio/ogg"), RawData);
+		}
+	}
+	else
+	{
+		TFuture<FSharedBuffer> FutureBuffer = SoundWave->RawData.GetPayload(USoundWave::FEditorAudioBulkData::EPayloadFlags::ReturnCopyWithTransformationData);
+		FutureBuffer.Wait();
+		if (FutureBuffer.IsReady())
+		{
+			FSharedBuffer Buffer = FutureBuffer.Get();
+			if (Buffer.GetSize() > 0)
+			{
+				return MakeAudioResult(TEXT("audio/wav"), MakeArrayView<uint8>(static_cast<uint8*>(const_cast<void*>(Buffer.GetData())), Buffer.GetSize()));
+			}
+		}
+	}
+#endif // WITH_EDITOR
+
+	return FModelContextProtocolToolResult();
+}
diff --git a/Engine/Plugins/Experimental/ModelContextProtocol/Source/ModelContextProtocolEngine/Private/ModelContextProtocolSettings.cpp b/Engine/Plugins/Experimental/ModelContextProtocol/Source/ModelContextProtocolEngine/Private/ModelContextProtocolSettings.cpp
new file mode 100644
--- /dev/null
+++ b/Engine/Plugins/Experimental/ModelContextProtocol/Source/ModelContextProtocolEngine/Private/ModelContextProtocolSettings.cpp
@@ -0,0 +1,38 @@
+﻿// Copyright Epic Games, Inc. All Rights Reserved.
+
+#include "ModelContextProtocolSettings.h"
+
+#include "Misc/CommandLine.h"
+#include "ModelContextProtocol.h"
+
+#include UE_INLINE_GENERATED_CPP_BY_NAME(ModelContextProtocolSettings)
+
+bool UE::ModelContextProtocol::ShouldAutoStartServer()
+{
+	if (FParse::Param(FCommandLine::Get(), TEXT("StartModelContextProtocolServer")))
+	{
+		return true;
+	}
+	
+	return GetDefault<UModelContextProtocolSettings>()->bAutoStartServer;
+}
+
+uint32 UE::ModelContextProtocol::GetServerPortNumber()
+{
+	uint32 Port = 0;
+	if (FParse::Value(FCommandLine::Get(), TEXT("ModelContextProtocolPort="), Port))
+	{
+		if (Port >= 1 && Port <= 65535)
+		{
+			return Port;
+		}
+		UE_LOG(LogModelContextProtocol, Warning, TEXT("Invalid command-line port %u. Must be 1-65535. Falling back to settings."), Port);
+	}
+
+	return GetDefault<UModelContextProtocolSettings>()->ServerPortNumber;
+}
+
+FString UE::ModelContextProtocol::GetServerUrlPath()
+{
+	return GetDefault<UModelContextProtocolSettings>()->ServerUrlPath;
+}
diff --git a/Engine/Plugins/Experimental/ModelContextProtocol/Source/ModelContextProtocolEngine/Private/ModelContextProtocolToolAsyncAction.cpp b/Engine/Plugins/Experimental/ModelContextProtocol/Source/ModelContextProtocolEngine/Private/ModelContextProtocolToolAsyncAction.cpp
new file mode 100644
--- /dev/null
+++ b/Engine/Plugins/Experimental/ModelContextProtocol/Source/ModelContextProtocolEngine/Private/ModelContextProtocolToolAsyncAction.cpp
@@ -0,0 +1,388 @@
+// Copyright Epic Games, Inc. All Rights Reserved.
+
+#include "ModelContextProtocolToolAsyncAction.h"
+#include "IModelContextProtocolModule.h"
+#include "ModelContextProtocolEngineMetaData.h"
+#include "ModelContextProtocolSession.h"
+#include "ModelContextProtocolToolUtils.h"
+
+#include "Engine/Engine.h"
+#include "Engine/World.h"
+#include "JsonSchema/JsonSchemaGenerator.h"
+#include "JsonSchema/JsonSchemaPropertyFilter.h"
+
+#if WITH_EDITORONLY_DATA
+#include "JsonSchema/JsonSchemaGeneratorEditor.h"
+#endif // WITH_EDITORONLY_DATA
+
+#if WITH_EDITOR
+#include "Editor/EditorEngine.h"
+#endif
+
+
+#include UE_INLINE_GENERATED_CPP_BY_NAME(ModelContextProtocolToolAsyncAction)
+
+void UModelContextProtocolToolAsyncAction::PostInitProperties()
+{
+	Super::PostInitProperties();
+
+	// Register with IModelContextProtocolModule
+	if (bAutoRegisterTool && HasAnyFlags(RF_ClassDefaultObject) && !GetClass()->HasAnyClassFlags(CLASS_Abstract) && GetClass()->HasAnyClassFlags(CLASS_Native))
+	{
+		RegisterTool();
+	}
+}
+
+void UModelContextProtocolToolAsyncAction::FinishDestroy()
+{
+	// Remove from global tool libraries list
+	if (HasAnyFlags(RF_ClassDefaultObject))
+	{
+		DeregisterTool();
+	}
+
+	Super::FinishDestroy();
+}
+
+void UModelContextProtocolToolAsyncAction::SetReadyToDestroy()
+{
+	Super::SetReadyToDestroy();
+
+	OnAsyncToolComplete.Broadcast(this);
+	OnAsyncToolComplete.Clear();
+}
+
+void UModelContextProtocolToolAsyncAction::RegisterTool()
+{
+	DeregisterTool();
+
+	check(!Tool.IsValid());
+	check(HasAnyFlags(RF_ClassDefaultObject));
+
+#if WITH_EDITOR
+	CollectToolMetaData();
+#endif
+
+	IModelContextProtocolModule& ModelContextProtocolModule = IModelContextProtocolModule::GetChecked();
+
+	if (UFunction* Function = GetClass()->FindFunctionByName(GetToolFunctionName(), EIncludeSuperFlag::ExcludeSuper))
+	{
+		Tool = MakeShared<FModelContextProtocolAsyncActionTool>(this, Function);
+
+		if (Tool.IsValid())
+		{
+			ModelContextProtocolModule.AddTool(Tool.ToSharedRef());
+		}
+	}
+
+	// Re-register tools on refresh
+	ModelContextProtocolModule.OnRefreshTools().AddWeakLambda(this, [this]()
+		{
+			RegisterTool();
+		});
+}
+
+void UModelContextProtocolToolAsyncAction::DeregisterTool()
+{
+	if (IModelContextProtocolModule* ModelContextProtocolModule = IModelContextProtocolModule::Get())
+	{
+		if (Tool.IsValid())
+		{
+			ModelContextProtocolModule->RemoveTool(Tool.ToSharedRef());
+		}
+
+		// Re-register tools on refresh
+		ModelContextProtocolModule->OnRefreshTools().RemoveAll(this);
+	}
+	Tool.Reset();
+}
+
+#if WITH_EDITORONLY_DATA
+void UModelContextProtocolToolAsyncAction::CollectToolMetaData()
+{
+	Modify();
+
+	ToolFunctionMetaData = FModelContextProtocolFunctionMetaData();
+	ToolResultMetaData = FModelContextProtocolFunctionMetaData();
+
+	if (UFunction* Function = GetClass()->FindFunctionByName(GetToolFunctionName(), EIncludeSuperFlag::ExcludeSuper))
+	{
+		FJsonSchemaPropertyFilter FunctionFilter(/*CheckFlags*/CPF_Parm | CPF_ReturnParm | CPF_OutParm | CPF_ConstParm);
+		ToolFunctionMetaData = UE::ModelContextProtocol::CollectFunctionMetaData(Function, FunctionFilter);
+	}
+
+	if (FProperty* Property = GetClass()->FindPropertyByName(GetToolResultPropertyName()))
+	{
+		FJsonSchemaPropertyFilter ClassFilter;
+		FJsonSchemaEditorMetadata EditorMetadata = FJsonSchemaGeneratorEditor::UStructToJsonSchemaMetadata(GetClass(), ClassFilter, nullptr);
+
+		if (EditorMetadata.RootStructMetadata.IsValid())
+		{
+			FString DescriptionString;
+			if (EditorMetadata.RootStructMetadata->TryGetStringField(TEXT("description"), DescriptionString) && !DescriptionString.IsEmpty())
+			{
+				ToolResultMetaData.Description = FText::FromString(DescriptionString);
+			}
+		}
+
+		for (const auto& [PropertyPath, PropertyMetadataJsonObject] : EditorMetadata.GetPropertyMemberPathToPropertyMetadataMap())
+		{
+			if (!PropertyMetadataJsonObject.IsValid())
+			{
+				continue;
+			}
+
+			FModelContextProtocolPropertyMetaData PropertyMetaData;
+			FString DescriptionString;
+			if (PropertyMetadataJsonObject->TryGetStringField(TEXT("description"), DescriptionString) && !DescriptionString.IsEmpty())
+			{
+				PropertyMetaData.Description = FText::FromString(DescriptionString);
+			}
+
+			double MinValue;
+			if (PropertyMetadataJsonObject->TryGetNumberField(TEXT("minimum"), MinValue))
+			{
+				PropertyMetaData.ClampMin = MinValue;
+			}
+
+			double MaxValue;
+			if (PropertyMetadataJsonObject->TryGetNumberField(TEXT("maximum"), MaxValue))
+			{
+				PropertyMetaData.ClampMax = MaxValue;
+			}
+
+			ToolResultMetaData.PropertyMetaData.Add(PropertyPath, MoveTemp(PropertyMetaData));
+		}
+	}
+}
+#endif // WITH_EDITORONLY_DATA
+
+FModelContextProtocolAsyncActionTool::FModelContextProtocolAsyncActionTool(UModelContextProtocolToolAsyncAction* InAction, UFunction* InFunction)
+	: AsyncAction(InAction)
+	, Function(InFunction)
+	, FunctionParamsContainer(Function)
+{
+	check(AsyncAction);
+	check(Function);
+}
+
+FString FModelContextProtocolAsyncActionTool::GetName() const
+{
+	check(Function);
+	return Function->GetName();
+}
+
+FString FModelContextProtocolAsyncActionTool::GetDescription() const
+{
+	check(AsyncAction);
+	check(Function);
+
+	const FModelContextProtocolFunctionMetaData* FunctionMetaData = AsyncAction->GetToolFunctionMetaData();
+	if (FunctionMetaData && FunctionMetaData->Description.IsSet())
+	{
+		return FunctionMetaData->Description->ToString();
+	}
+
+	return FString();
+}
+
+TSharedPtr<FJsonObject> FModelContextProtocolAsyncActionTool::GetInputJsonSchema() const
+{
+	check(AsyncAction);
+	check(Function);
+
+	// Return cached?
+	if (InputJsonSchema.IsValid())
+	{
+		return InputJsonSchema;
+	}
+
+	const FModelContextProtocolFunctionMetaData* FunctionMetaData = AsyncAction->GetToolFunctionMetaData();
+
+	// Build property filter — skip WorldContext param as we auto-set it, and skip return/out params
+	TOptional<TSet<FString>> SkipPropertyMemberPaths;
+	if (FunctionMetaData && !FunctionMetaData->WorldContext.IsEmpty())
+	{
+		SkipPropertyMemberPaths.Emplace();
+		SkipPropertyMemberPaths->Add(FunctionMetaData->WorldContext);
+	}
+
+	FJsonSchemaPropertyFilter PropertyFilter(
+		/*CheckFlags*/CPF_Parm | CPF_ConstParm,
+		/*SkipFlags*/CPF_ReturnParm | CPF_OutParm,
+		/*RequiredPropertyMemberPaths*/{},
+		SkipPropertyMemberPaths);
+
+	// Convert cookable metadata to engine format for schema generation
+	TOptional<FJsonSchemaEditorMetadata> CachedEditorMetadata;
+	if (FunctionMetaData)
+	{
+		CachedEditorMetadata = UE::ModelContextProtocol::ConvertToCachedEditorMetadata(*FunctionMetaData);
+	}
+
+	TSharedPtr<FJsonObject> SchemaResult = FJsonSchemaGenerator::UStructToJsonSchemaObject(
+		Function, PropertyFilter, CachedEditorMetadata.IsSet() ? &CachedEditorMetadata.GetValue() : nullptr);
+
+	// UStructToJsonSchemaObject returns JSON-RPC wrapper for UFunctions: {inputSchema, outputSchema}
+	// Extract inputSchema
+	if (SchemaResult.IsValid())
+	{
+		const TSharedPtr<FJsonObject>* InputSchemaField;
+		if (SchemaResult->TryGetObjectField(TEXT("inputSchema"), InputSchemaField))
+		{
+			InputJsonSchema = *InputSchemaField;
+		}
+	}
+
+	return InputJsonSchema;
+}
+
+TSharedPtr<FJsonObject> FModelContextProtocolAsyncActionTool::GetOutputJsonSchema() const
+{
+	CacheResultTypeInfo();
+
+	return OutputJsonSchema;
+}
+
+void FModelContextProtocolAsyncActionTool::RunAsync(const FModelContextProtocolToolRequestId& RequestId, const TSharedPtr<FJsonObject>& Params, const FResultCallback& OnComplete)
+{
+	check(AsyncAction);
+	check(Function);
+
+	CacheResultTypeInfo();
+
+	// Prepare params
+	Function->InitializeStruct(FunctionParamsContainer.GetStructMemory());
+	ON_SCOPE_EXIT
+	{
+		Function->DestroyStruct(FunctionParamsContainer.GetStructMemory());
+	};
+
+	// Apply param defaults
+	const FModelContextProtocolFunctionMetaData* FunctionMetaData = AsyncAction->GetToolFunctionMetaData();
+
+	if (FunctionMetaData)
+	{
+		for (TFieldIterator<FProperty> ParamPropertyIt(Function); ParamPropertyIt; ++ParamPropertyIt)
+		{
+			const FProperty* ParamProperty = *ParamPropertyIt;
+			FString ParamName = ParamProperty->GetAuthoredName();
+			if (const FModelContextProtocolPropertyMetaData* PropertyMetaData = FunctionMetaData->PropertyMetaData.Find(*ParamName);
+				PropertyMetaData && PropertyMetaData->DefaultValue.IsSet())
+			{
+				ParamProperty->ImportText_InContainer(**PropertyMetaData->DefaultValue, FunctionParamsContainer.GetStructMemory(), AsyncAction, PPF_None);
+			}
+		}
+	}
+
+	if (Params.IsValid())
+	{
+		// Json Params -> UFunction params container
+		FText FailReason;
+		if (!FJsonObjectConverter::JsonObjectToUStruct(Params.ToSharedRef(), Function, FunctionParamsContainer.GetStructMemory(), /*CheckFlags*/CPF_Parm | CPF_ConstParm, /*SkipFlags*/CPF_ReturnParm | CPF_OutParm, /*bStrictMode*/false, &FailReason))
+		{
+			OnComplete(UE::ModelContextProtocol::MakeErrorResult(FailReason.ToString()));
+			return;
+		}
+	}
+
+	// WorldContext
+	if (FunctionMetaData && !FunctionMetaData->WorldContext.IsEmpty())
+	{
+		UObject* WorldContextObject = GWorld.GetReference();
+		if (FProperty* WorldContextProperty = Function->FindPropertyByName(FName(FunctionMetaData->WorldContext)); ensure(WorldContextProperty))
+		{
+			if (FObjectPropertyBase* WorldContextObjectProperty = CastField<FObjectPropertyBase>(WorldContextProperty); ensure(WorldContextObjectProperty))
+			{
+				WorldContextObjectProperty->SetObjectPropertyValue_InContainer(FunctionParamsContainer.GetStructMemory(), WorldContextObject);
+			}
+		}
+	}
+
+	{
+		// Execute tool
+#if WITH_EDITOR
+		FEditorScriptExecutionGuard ScriptExecutionGuard;
+#endif
+		AsyncAction->ProcessEvent(Function, FunctionParamsContainer.GetStructMemory());
+	}
+
+	if (FProperty* ReturnProperty = Function->GetReturnProperty(); ReturnProperty && ensure(Function->ReturnValueOffset != MAX_uint16))
+	{
+		if (FObjectProperty* ObjectReturnProperty = CastField<FObjectProperty>(ReturnProperty); ensure(ObjectReturnProperty))
+		{
+			if (UObject* ObjectResult = ObjectReturnProperty->GetObjectPropertyValue_InContainer(FunctionParamsContainer.GetStructMemory()))
+			{
+				if (UModelContextProtocolToolAsyncAction* NewAction = Cast<UModelContextProtocolToolAsyncAction>(ObjectResult))
+				{
+					NewAction->ToolRequestId = RequestId;
+
+					InProgressActions.Add(NewAction);
+
+					NewAction->OnAsyncToolComplete.AddSPLambda(this, [this, OnComplete](UModelContextProtocolToolAsyncAction* InAction)
+						{
+							InProgressActions.Remove(InAction);
+
+							FModelContextProtocolToolResult Result;
+
+							if (FProperty* ResultProperty = InAction->GetClass()->FindPropertyByName(InAction->GetToolResultPropertyName()))
+							{
+								check(ResultType.IsSet());
+
+								Result = UE::ModelContextProtocol::GetToolResultFromType(*ResultType, ResultProperty, InAction, ResultProperty->GetOffset_ForInternal());
+							}
+
+							OnComplete(Result);
+						});
+
+					return;
+				}
+			}
+		}
+	}
+
+	OnComplete(UE::ModelContextProtocol::MakeErrorResult(TEXT("Error")));
+}
+
+void FModelContextProtocolAsyncActionTool::CancelAsync(const FModelContextProtocolToolRequestId& RequestId)
+{
+	if (RequestId.IsValid())
+	{
+		for (TObjectPtr<UModelContextProtocolToolAsyncAction>& Action : InProgressActions)
+		{
+			if (Action && Action->ToolRequestId.IsValid() && FJsonValue::CompareEqual(*RequestId.RequestId, *Action->ToolRequestId.RequestId))
+			{
+				Action->Cancel();
+				break;
+			}
+		}
+	}
+}
+
+void FModelContextProtocolAsyncActionTool::AddReferencedObjects(FReferenceCollector& Collector)
+{
+	Collector.AddReferencedObject(AsyncAction);
+	Collector.AddReferencedObject(Function);
+	Collector.AddReferencedObjects(InProgressActions);
+}
+
+void FModelContextProtocolAsyncActionTool::CacheResultTypeInfo() const
+{
+	check(AsyncAction);
+	check(Function);
+
+	// Already cached?
+	if (ResultType.IsSet())
+	{
+		return;
+	}
+
+	// Determine result type
+	ResultType = EModelContextProtocolToolResultType::None;
+	OutputJsonSchema.Reset();
+
+	if (FProperty* ResultProperty = AsyncAction->GetClass()->FindPropertyByName(AsyncAction->GetToolResultPropertyName()))
+	{
+		ResultType = UE::ModelContextProtocol::GetToolResultType(ResultProperty, AsyncAction->GetToolResultMetaData(), OutputJsonSchema);
+	}
+}
diff --git a/Engine/Plugins/Experimental/ModelContextProtocol/Source/ModelContextProtocolEngine/Private/ModelContextProtocolToolLibrary.cpp b/Engine/Plugins/Experimental/ModelContextProtocol/Source/ModelContextProtocolEngine/Private/ModelContextProtocolToolLibrary.cpp
new file mode 100644
--- /dev/null
+++ b/Engine/Plugins/Experimental/ModelContextProtocol/Source/ModelContextProtocolEngine/Private/ModelContextProtocolToolLibrary.cpp
@@ -0,0 +1,350 @@
+// Copyright Epic Games, Inc. All Rights Reserved.
+
+#include "ModelContextProtocolToolLibrary.h"
+#include "IModelContextProtocolModule.h"
+#include "ModelContextProtocol.h"
+#include "ModelContextProtocolEngineMetaData.h"
+#include "ModelContextProtocolToolUtils.h"
+#include "JsonSchema/JsonSchemaGenerator.h"
+#include "JsonSchema/JsonSchemaPropertyFilter.h"
+#include "Engine/World.h"
+#include "UObject/UObjectIterator.h"
+
+#if WITH_EDITOR
+#include "Editor.h"
+#include "ScopedTransaction.h"
+#endif // WITH_EDITOR
+
+#include UE_INLINE_GENERATED_CPP_BY_NAME(ModelContextProtocolToolLibrary)
+
+#define LOCTEXT_NAMESPACE "ModelContextProtocolToolLibrary"
+
+void UModelContextProtocolToolLibrary::PostInitProperties()
+{
+	Super::PostInitProperties();
+
+	// Note: We only need to register tools for Native classes here, as BP classes will register later via PostLoad or PostCDOCompiled
+	if (bAutoRegisterTools && HasAnyFlags(RF_ClassDefaultObject) && !GetClass()->HasAnyClassFlags(CLASS_Abstract) && GetClass()->HasAnyClassFlags(CLASS_Native))
+	{
+		RegisterTools();
+	}
+}
+
+void UModelContextProtocolToolLibrary::PostLoad()
+{
+	Super::PostLoad();
+	
+	// Register BP classes here (native classes already registered in PostInitProperties)
+	if (bAutoRegisterTools && HasAnyFlags(RF_ClassDefaultObject) && !GetClass()->HasAnyClassFlags(CLASS_Abstract) && !GetClass()->HasAnyClassFlags(CLASS_Native))
+	{
+		RegisterTools();
+	}
+}
+
+#if WITH_EDITOR
+void UModelContextProtocolToolLibrary::PostCDOCompiled(const FPostCDOCompiledContext& Context)
+{
+	Super::PostCDOCompiled(Context);
+
+	// @todo Account for CDO recompilation for bAutoRegisterTools=false but where manual RegisterTools has already occurred?   
+	if (bAutoRegisterTools)
+	{
+		RegisterTools();
+	}
+}
+#endif // WITH_EDITOR
+
+void UModelContextProtocolToolLibrary::FinishDestroy()
+{
+	// Remove from global tool libraries list
+	if (HasAnyFlags(RF_ClassDefaultObject))
+	{
+		DeregisterTools();
+	}
+	
+	Super::FinishDestroy();
+}
+
+void UModelContextProtocolToolLibrary::RegisterTools()
+{
+	DeregisterTools();
+	
+	check(Tools.IsEmpty());
+	check(HasAnyFlags(RF_ClassDefaultObject));
+
+#if WITH_EDITOR
+	CollectFunctionMetaData();
+#endif
+
+	IModelContextProtocolModule& ModelContextProtocolModule = IModelContextProtocolModule::GetChecked();
+	
+	// Define tools for all public functions
+	for (TFieldIterator<UFunction> FunctionIter(GetClass(), EFieldIteratorFlags::ExcludeSuper); FunctionIter; ++FunctionIter)
+	{
+		UFunction* Function = *FunctionIter;
+		const FName FunctionName = Function->GetFName();
+		if (FunctionName == NAME_ExecuteUbergraph
+			|| !FunctionIter->HasAnyFunctionFlags(FUNC_Public))
+		{
+			continue;
+		}
+
+		TSharedRef<FModelContextProtocolLibraryTool> Tool = MakeShared<FModelContextProtocolLibraryTool>(this, Function); 
+		Tools.Add(Tool);
+
+		ModelContextProtocolModule.AddTool(Tool);
+	}
+
+	// Re-register tools on refresh
+	ModelContextProtocolModule.OnRefreshTools().AddWeakLambda(this, [this]()
+		{
+			RegisterTools();
+		});
+}
+
+void UModelContextProtocolToolLibrary::DeregisterTools()
+{
+	if (IModelContextProtocolModule* ModelContextProtocolModule = IModelContextProtocolModule::Get())
+	{
+		for (const TSharedRef<FModelContextProtocolLibraryTool>& Tool : Tools)
+		{
+			ModelContextProtocolModule->RemoveTool(Tool);
+		}
+
+		// Re-register tools on refresh
+		ModelContextProtocolModule->OnRefreshTools().RemoveAll(this);
+	}
+	Tools.Reset();
+}
+
+#if WITH_EDITORONLY_DATA
+void UModelContextProtocolToolLibrary::CollectFunctionMetaData()
+{
+	Modify();
+
+	FunctionMetaData.Reset();
+	for (TFieldIterator<UFunction> FunctionIter(GetClass(), EFieldIteratorFlags::ExcludeSuper); FunctionIter; ++FunctionIter)
+	{
+		const UFunction* Function = *FunctionIter;
+		const FName FunctionName = Function->GetFName();
+
+		// Skip hidden functions
+		if (FunctionName == NAME_ExecuteUbergraph
+			|| !FunctionIter->HasAnyFunctionFlags(FUNC_Public))
+		{
+			continue;
+		}
+
+		FJsonSchemaPropertyFilter PropertyFilter(/*CheckFlags*/CPF_Parm | CPF_ReturnParm | CPF_OutParm | CPF_ConstParm);
+		FunctionMetaData.Add(FunctionName, UE::ModelContextProtocol::CollectFunctionMetaData(Function, PropertyFilter));
+	}
+}
+#endif // WITH_EDITORONLY_DATA
+
+UModelContextProtocolToolLibraryBlueprint::UModelContextProtocolToolLibraryBlueprint()
+{
+	BlueprintType = BPTYPE_FunctionLibrary;
+	ParentClass = UModelContextProtocolToolLibrary::StaticClass();
+#if WITH_EDITORONLY_DATA
+	ShouldCookPropertyGuidsValue = EShouldCookBlueprintPropertyGuids::Yes;
+#endif // WITH_EDITORONLY_DATA
+}
+
+#if WITH_EDITOR
+bool UModelContextProtocolToolLibraryBlueprint::AlwaysCompileOnLoad() const
+{
+	return true;
+}
+#endif
+
+FModelContextProtocolLibraryTool::FModelContextProtocolLibraryTool(UModelContextProtocolToolLibrary* InLibrary, UFunction* InFunction)
+: Library(InLibrary)
+, Function(InFunction)
+, FunctionParamsContainer(Function)
+{
+	check(Library);
+	check(Function);
+}
+
+FString FModelContextProtocolLibraryTool::GetName() const
+{
+	check(Function);
+	return Function->GetName();
+}
+
+FString FModelContextProtocolLibraryTool::GetDescription() const
+{
+	check(Library);
+	check(Function);
+	if (const FModelContextProtocolFunctionMetaData* FunctionMetaData = Library->FindFunctionMetaData(Function->GetFName()))
+	{
+		if (FunctionMetaData->Description.IsSet())
+		{
+			return FunctionMetaData->Description->ToString();
+		}
+	}
+
+	return FString();
+}
+
+TSharedPtr<FJsonObject> FModelContextProtocolLibraryTool::GetInputJsonSchema() const
+{
+	check(Library);
+	check(Function);
+
+	// Return cached?
+	if (InputJsonSchema.IsValid())
+	{
+		return InputJsonSchema;
+	}
+
+	const FModelContextProtocolFunctionMetaData* FunctionMetaData = Library->FindFunctionMetaData(Function->GetFName());
+
+	// Build property filter — skip WorldContext param as we auto-set it, and skip return/out params
+	TOptional<TSet<FString>> SkipPropertyMemberPaths;
+	if (FunctionMetaData && !FunctionMetaData->WorldContext.IsEmpty())
+	{
+		SkipPropertyMemberPaths.Emplace();
+		SkipPropertyMemberPaths->Add(FunctionMetaData->WorldContext);
+	}
+
+	FJsonSchemaPropertyFilter PropertyFilter(
+		/*CheckFlags*/CPF_Parm | CPF_ConstParm,
+		/*SkipFlags*/CPF_ReturnParm | CPF_OutParm,
+		/*RequiredPropertyMemberPaths*/{},
+		SkipPropertyMemberPaths);
+
+	// Convert cookable metadata to engine format for schema generation
+	TOptional<FJsonSchemaEditorMetadata> CachedEditorMetadata;
+	if (FunctionMetaData)
+	{
+		CachedEditorMetadata = UE::ModelContextProtocol::ConvertToCachedEditorMetadata(*FunctionMetaData);
+	}
+
+	TSharedPtr<FJsonObject> SchemaResult = FJsonSchemaGenerator::UStructToJsonSchemaObject(
+		Function, PropertyFilter, CachedEditorMetadata.IsSet() ? &CachedEditorMetadata.GetValue() : nullptr);
+
+	// UStructToJsonSchemaObject returns JSON-RPC wrapper for UFunctions: {inputSchema, outputSchema}
+	// Extract inputSchema
+	if (SchemaResult.IsValid())
+	{
+		const TSharedPtr<FJsonObject>* InputSchemaField;
+		if (SchemaResult->TryGetObjectField(TEXT("inputSchema"), InputSchemaField))
+		{
+			InputJsonSchema = *InputSchemaField;
+		}
+	}
+
+	return InputJsonSchema;
+}
+
+TSharedPtr<FJsonObject> FModelContextProtocolLibraryTool::GetOutputJsonSchema() const
+{
+	CacheResultTypeInfo();
+	
+	return OutputJsonSchema;
+}
+
+FModelContextProtocolToolResult FModelContextProtocolLibraryTool::Run(const TSharedPtr<FJsonObject>& Params)
+{
+#if WITH_EDITOR
+	const FScopedTransaction Transaction(FText::FromString(FName::NameToDisplayString(GetName(), /*bIsBool*/false)));
+#endif // WITH_EDITOR
+	
+	check(Library);
+	check(Function);
+
+	CacheResultTypeInfo();
+
+	// Prepare params
+	Function->InitializeStruct(FunctionParamsContainer.GetStructMemory());
+	ON_SCOPE_EXIT
+	{
+		Function->DestroyStruct(FunctionParamsContainer.GetStructMemory());
+	};
+
+	// Apply param defaults
+	const FModelContextProtocolFunctionMetaData* FunctionMetaData = Library->FindFunctionMetaData(Function->GetFName());
+	if (FunctionMetaData)
+	{
+		for (TFieldIterator<FProperty> ParamPropertyIt(Function); ParamPropertyIt; ++ParamPropertyIt)
+		{
+			const FProperty* ParamProperty = *ParamPropertyIt;
+			FString ParamName = ParamProperty->GetAuthoredName();
+			if (const FModelContextProtocolPropertyMetaData* PropertyMetaData = FunctionMetaData->PropertyMetaData.Find(*ParamName);
+				PropertyMetaData && PropertyMetaData->DefaultValue.IsSet())
+			{
+				ParamProperty->ImportText_InContainer(**PropertyMetaData->DefaultValue, FunctionParamsContainer.GetStructMemory(), Library, PPF_None);
+			}
+		}
+	}
+	
+	if (Params.IsValid())
+	{
+		// Json Params -> UFunction params container
+		FText FailReason;
+		if (!FJsonObjectConverter::JsonObjectToUStruct(Params.ToSharedRef(), Function, FunctionParamsContainer.GetStructMemory(), /*CheckFlags*/CPF_Parm | CPF_ConstParm, /*SkipFlags*/CPF_ReturnParm | CPF_OutParm, /*bStrictMode*/false, &FailReason))
+		{
+			return UE::ModelContextProtocol::MakeErrorResult(FailReason.ToString());
+		}
+	}
+
+	// WorldContext
+	if (FunctionMetaData && !FunctionMetaData->WorldContext.IsEmpty())
+	{
+		UObject* WorldContextObject = GWorld.GetReference();
+		if (FProperty* WorldContextProperty = Function->FindPropertyByName(FName(FunctionMetaData->WorldContext)); ensure(WorldContextProperty))
+		{
+			if (FObjectPropertyBase* WorldContextObjectProperty = CastField<FObjectPropertyBase>(WorldContextProperty); ensure(WorldContextObjectProperty))
+			{
+				WorldContextObjectProperty->SetObjectPropertyValue_InContainer(FunctionParamsContainer.GetStructMemory(), WorldContextObject);
+			}
+		}
+	}
+	
+	{
+		// Execute tool
+		FEditorScriptExecutionGuard ScriptExecutionGuard;
+		Library->ProcessEvent(Function, FunctionParamsContainer.GetStructMemory());
+	}
+
+	if (FProperty* ReturnProperty = Function->GetReturnProperty(); ReturnProperty && ensure(Function->ReturnValueOffset != MAX_uint16))
+	{
+		check(ResultType.IsSet());
+
+		return UE::ModelContextProtocol::GetToolResultFromType(*ResultType, ReturnProperty, FunctionParamsContainer.GetStructMemory(), Function->ReturnValueOffset);
+	}
+
+	// No result (void return)
+	return FModelContextProtocolToolResult();
+}
+
+void FModelContextProtocolLibraryTool::AddReferencedObjects(FReferenceCollector& Collector)
+{
+	Collector.AddReferencedObject(Library);
+	Collector.AddReferencedObject(Function);
+}
+
+void FModelContextProtocolLibraryTool::CacheResultTypeInfo() const
+{
+	check(Library);
+	check(Function);
+
+	// Already cached?
+	if (ResultType.IsSet())
+	{
+		return;
+	}
+
+	// Determine result type
+	ResultType = EModelContextProtocolToolResultType::None;
+	OutputJsonSchema.Reset();
+	if (FProperty* ReturnProperty = Function->GetReturnProperty(); ReturnProperty && ensure(Function->ReturnValueOffset != MAX_uint16))
+	{
+		const FModelContextProtocolFunctionMetaData* FunctionMetaData = Library->FindFunctionMetaData(Function->GetFName());
+
+		ResultType = UE::ModelContextProtocol::GetToolResultType(ReturnProperty, FunctionMetaData, OutputJsonSchema);
+	}
+}
+
+#undef LOCTEXT_NAMESPACE
diff --git a/Engine/Plugins/Experimental/ModelContextProtocol/Source/ModelContextProtocolEngine/Private/ModelContextProtocolToolUtils.cpp b/Engine/Plugins/Experimental/ModelContextProtocol/Source/ModelContextProtocolEngine/Private/ModelContextProtocolToolUtils.cpp
new file mode 100644
--- /dev/null
+++ b/Engine/Plugins/Experimental/ModelContextProtocol/Source/ModelContextProtocolEngine/Private/ModelContextProtocolToolUtils.cpp
@@ -0,0 +1,140 @@
+// Copyright Epic Games, Inc. All Rights Reserved.
+
+#include "ModelContextProtocolToolUtils.h"
+#include "ModelContextProtocol.h"
+#include "ModelContextProtocolEngineMetaData.h"
+#include "ModelContextProtocolEngineToolResults.h"
+#include "ModelContextProtocolMetaData.h"
+#include "ModelContextProtocolToolResults.h"
+
+#include "Engine/Texture2D.h"
+#include "JsonSchema/JsonSchemaGenerator.h"
+#include "JsonSchema/JsonSchemaPropertyFilter.h"
+#include "Sound/SoundWave.h"
+#include "UObject/TextProperty.h"
+#include "UObject/Utf8StrProperty.h"
+
+namespace UE::ModelContextProtocol
+{
+	EModelContextProtocolToolResultType GetToolResultType(FProperty* Property, const FModelContextProtocolFunctionMetaData* MetaData, TSharedPtr<FJsonObject>& OutSchema)
+	{
+		OutSchema.Reset();
+
+		if (Property)
+		{
+			// Image result?
+			if (const FObjectProperty* ObjectReturnProperty = CastField<FObjectProperty>(Property);
+				ObjectReturnProperty && ObjectReturnProperty->PropertyClass == UTexture2D::StaticClass())
+			{
+				return EModelContextProtocolToolResultType::Image;
+			}
+
+			// Audio result?
+			if (const FObjectProperty* ObjectReturnProperty = CastField<FObjectProperty>(Property);
+				ObjectReturnProperty && ObjectReturnProperty->PropertyClass == USoundWave::StaticClass())
+			{
+				return EModelContextProtocolToolResultType::Audio;
+			}
+
+			// Text result?
+			if (Property->IsA<FStrProperty>()
+				|| Property->IsA<FTextProperty>()
+				|| Property->IsA<FUtf8StrProperty>())
+			{
+				return EModelContextProtocolToolResultType::Text;
+			}
+
+			// We want to convert soft object references to text paths
+			if (Property->IsA<FSoftObjectProperty>())
+			{
+				return EModelContextProtocolToolResultType::Text;
+			}
+
+			// Return structured result
+			// @todo Support reference output vars
+			{
+				FJsonSchemaPropertyFilter PropertyFilter;
+				TOptional<FJsonSchemaEditorMetadata> CachedEditorMetadata;
+				if (MetaData)
+				{
+					CachedEditorMetadata = UE::ModelContextProtocol::ConvertToCachedEditorMetadata(*MetaData);
+				}
+				OutSchema = FJsonSchemaGenerator::FPropertyToJsonSchemaObject(
+					Property, PropertyFilter, CachedEditorMetadata.IsSet() ? &CachedEditorMetadata.GetValue() : nullptr);
+			}
+
+			// {"result":} object wrapper for POD results?
+			if (UE::ModelContextProtocol::bWrapPODResultsInObject && OutSchema.IsValid())
+			{
+				FString TypeString;
+
+				if (OutSchema->TryGetStringField(TEXT("type"), TypeString))
+				{
+					if (TypeString != TEXT("object"))
+					{
+						FJsonDomBuilder::FObject ResultWrapperObjectSchema;
+						ResultWrapperObjectSchema.Set(TEXT("type"), TEXT("object"));
+						ResultWrapperObjectSchema.Set(TEXT("properties"), FJsonDomBuilder::FObject().Set(UE::ModelContextProtocol::PODWrapperResultPropertyName, OutSchema));
+						ResultWrapperObjectSchema.Set(TEXT("required"), FJsonDomBuilder::FArray().Add(FString(UE::ModelContextProtocol::PODWrapperResultPropertyName)));
+
+						OutSchema = ResultWrapperObjectSchema.AsJsonObject();
+					}
+				}
+			}
+
+			return EModelContextProtocolToolResultType::StructuredContent;
+		}
+
+		return EModelContextProtocolToolResultType::None;
+	}
+
+	FModelContextProtocolToolResult GetToolResultFromType(EModelContextProtocolToolResultType ResultType, FProperty* Property, void* Container, uint16 ReturnValueOffset)
+	{
+		if (ResultType == EModelContextProtocolToolResultType::Image)
+		{
+			if (FObjectProperty* ObjectReturnProperty = CastField<FObjectProperty>(Property); ensure(ObjectReturnProperty) && ensure(ObjectReturnProperty->PropertyClass == UTexture2D::StaticClass()))
+			{
+				if (const UObject* ObjectResult = ObjectReturnProperty->GetObjectPropertyValue_InContainer(Container))
+				{
+					if (const UTexture2D* TextureObjectResult = Cast<UTexture2D>(ObjectResult); ensure(TextureObjectResult))
+					{
+						return UE::ModelContextProtocol::MakeImageResult(TextureObjectResult);
+					}
+				}
+			}
+		}
+		else if (ResultType == EModelContextProtocolToolResultType::Audio)
+		{
+			if (FObjectProperty* ObjectReturnProperty = CastField<FObjectProperty>(Property); ensure(ObjectReturnProperty) && ensure(ObjectReturnProperty->PropertyClass == USoundWave::StaticClass()))
+			{
+				if (UObject* ObjectResult = ObjectReturnProperty->GetObjectPropertyValue_InContainer(Container))
+				{
+					if (USoundWave* SoundObjectResult = Cast<USoundWave>(ObjectResult); ensure(SoundObjectResult))
+					{
+						return UE::ModelContextProtocol::MakeAudioResult(SoundObjectResult);
+					}
+				}
+			}
+		}
+		else if (ResultType == EModelContextProtocolToolResultType::Text)
+		{
+			// Export text with PPF_PropertyWindow PortFlags to ensure FText is exported as simple string
+			FString StringResult;
+			Property->ExportTextItem_InContainer(StringResult, Container, /*DefaultValue*/nullptr, /*Parent*/nullptr, PPF_PropertyWindow);
+
+			return UE::ModelContextProtocol::MakeTextResult(StringResult);
+		}
+		else if (ResultType == EModelContextProtocolToolResultType::StructuredContent)
+		{
+			// Return structured result
+			// @todo Support for OutParm's using UStructToJson with OutParm CheckFlags on FunctionParamsContainer
+			return UE::ModelContextProtocol::MakeStructuredContentResult(Property, (uint8*)Container + ReturnValueOffset);
+		}
+		else
+		{
+			ensureMsgf(false, TEXT("Unhandled result type %s"), *StaticEnum<EModelContextProtocolToolResultType>()->GetValueAsString(ResultType));
+		}
+
+		return FModelContextProtocolToolResult();
+	}
+};
\ No newline at end of file
diff --git a/Engine/Plugins/Experimental/ModelContextProtocol/Source/ModelContextProtocolEngine/Public/ModelContextProtocolEngineMetaData.h b/Engine/Plugins/Experimental/ModelContextProtocol/Source/ModelContextProtocolEngine/Public/ModelContextProtocolEngineMetaData.h
new file mode 100644
--- /dev/null
+++ b/Engine/Plugins/Experimental/ModelContextProtocol/Source/ModelContextProtocolEngine/Public/ModelContextProtocolEngineMetaData.h
@@ -0,0 +1,30 @@
+// Copyright Epic Games, Inc. All Rights Reserved.
+
+#pragma once
+
+#include "ModelContextProtocolMetaData.h"
+
+struct FJsonSchemaEditorMetadata;
+struct FJsonSchemaPropertyFilter;
+
+namespace UE::ModelContextProtocol
+{
+#if WITH_EDITORONLY_DATA
+	/**
+	 * Collects editor-only metadata for a UFunction, storing it in cookable form.
+	 *
+	 * @param Function The function to collect metadata for.
+	 * @param PropertyFilter Property filter for schema generation.
+	 * @return Collected function metadata.
+	 */
+	MODELCONTEXTPROTOCOLENGINE_API FModelContextProtocolFunctionMetaData CollectFunctionMetaData(const UFunction* Function, const FJsonSchemaPropertyFilter& PropertyFilter);
+#endif
+
+	/**
+	 * Converts cookable function metadata back into FJsonSchemaEditorMetadata for passing to FJsonSchemaGenerator.
+	 *
+	 * @param FunctionMetaData The cookable metadata to convert.
+	 * @return Engine editor metadata suitable for schema generation.
+	 */
+	MODELCONTEXTPROTOCOLENGINE_API FJsonSchemaEditorMetadata ConvertToCachedEditorMetadata(const FModelContextProtocolFunctionMetaData& FunctionMetaData);
+}
diff --git a/Engine/Plugins/Experimental/ModelContextProtocol/Source/ModelContextProtocolEngine/Public/ModelContextProtocolEngineToolResults.h b/Engine/Plugins/Experimental/ModelContextProtocol/Source/ModelContextProtocolEngine/Public/ModelContextProtocolEngineToolResults.h
new file mode 100644
--- /dev/null
+++ b/Engine/Plugins/Experimental/ModelContextProtocol/Source/ModelContextProtocolEngine/Public/ModelContextProtocolEngineToolResults.h
@@ -0,0 +1,20 @@
+// Copyright Epic Games, Inc. All Rights Reserved.
+
+#pragma once
+
+#include "ModelContextProtocolToolResults.h"
+
+class USoundWave;
+class UTexture2D;
+struct FColor;
+struct FImageView;
+
+namespace UE::ModelContextProtocol
+{
+	// @see https://modelcontextprotocol.io/specification/2025-06-18/server/tools#image-content
+	MODELCONTEXTPROTOCOLENGINE_API FModelContextProtocolToolResult MakeImageResult(const UTexture2D* Texture, const TCHAR* ToFormatExtension = TEXT("jpeg"), EModelContextProtocolAudience Audience = EModelContextProtocolAudience::All);
+	MODELCONTEXTPROTOCOLENGINE_API FModelContextProtocolToolResult MakeImageResult(const FImageView& Image, const TCHAR* ToFormatExtension = TEXT("jpeg"), EModelContextProtocolAudience Audience = EModelContextProtocolAudience::All);
+
+	// @see https://modelcontextprotocol.io/specification/2025-06-18/server/tools#audio-content
+	MODELCONTEXTPROTOCOLENGINE_API FModelContextProtocolToolResult MakeAudioResult(USoundWave* SoundWave);
+}
diff --git a/Engine/Plugins/Experimental/ModelContextProtocol/Source/ModelContextProtocolEngine/Public/ModelContextProtocolSettings.h b/Engine/Plugins/Experimental/ModelContextProtocol/Source/ModelContextProtocolEngine/Public/ModelContextProtocolSettings.h
new file mode 100644
--- /dev/null
+++ b/Engine/Plugins/Experimental/ModelContextProtocol/Source/ModelContextProtocolEngine/Public/ModelContextProtocolSettings.h
@@ -0,0 +1,36 @@
+﻿// Copyright Epic Games, Inc. All Rights Reserved.
+
+#pragma once
+
+#include "Engine/DeveloperSettings.h"
+
+#include "ModelContextProtocolSettings.generated.h"
+
+namespace UE::ModelContextProtocol
+{
+	MODELCONTEXTPROTOCOLENGINE_API bool ShouldAutoStartServer();
+	MODELCONTEXTPROTOCOLENGINE_API uint32 GetServerPortNumber();
+	MODELCONTEXTPROTOCOLENGINE_API FString GetServerUrlPath();
+}
+
+UCLASS(MinimalAPI, config=EditorPerProjectUserSettings, defaultconfig, meta=(DisplayName="Model Context Protocol"))
+class UModelContextProtocolSettings : public UDeveloperSettings
+{
+	GENERATED_BODY()
+public:
+
+	/** The URL base path to serve from e.g: "/mcp" -> http://localhost/mcp */
+	UPROPERTY(Config, EditAnywhere, Category="Server")
+	FString ServerUrlPath = TEXT("/mcp");
+	
+	/** The port number to serve from e.g: 8000 -> http://localhost:8000/mcp */
+	UPROPERTY(Config, EditAnywhere, Category="Server")
+	uint32 ServerPortNumber = 8000;
+
+	/**
+	 * If true, the HTTP server route will be automatically registered and HTTP listeners started during module startup.
+	 */
+	UPROPERTY(Config, EditAnywhere, Category="Server")
+	bool bAutoStartServer = false;
+};
+
diff --git a/Engine/Plugins/Experimental/ModelContextProtocol/Source/ModelContextProtocolEngine/Public/ModelContextProtocolToolAsyncAction.h b/Engine/Plugins/Experimental/ModelContextProtocol/Source/ModelContextProtocolEngine/Public/ModelContextProtocolToolAsyncAction.h
new file mode 100644
--- /dev/null
+++ b/Engine/Plugins/Experimental/ModelContextProtocol/Source/ModelContextProtocolEngine/Public/ModelContextProtocolToolAsyncAction.h
@@ -0,0 +1,109 @@
+// Copyright Epic Games, Inc. All Rights Reserved.
+
+#pragma once
+
+#include "Engine/CancellableAsyncAction.h"
+#include "IModelContextProtocolTool.h"
+#include "ModelContextProtocolMetaData.h"
+#include "ModelContextProtocolSession.h"
+#include "UObject/StructOnScope.h"
+
+#include "ModelContextProtocolToolAsyncAction.generated.h"
+
+struct FModelContextProtocolAsyncActionTool;
+
+#define UE_API MODELCONTEXTPROTOCOLENGINE_API
+
+DECLARE_MULTICAST_DELEGATE_OneParam(FOnModelContextProtocolAsyncToolComplete, UModelContextProtocolToolAsyncAction* /* AsyncAction */);
+
+/**
+ * @deprecated Use UToolsetDefinition (ToolsetRegistry plugin) instead. UModelContextProtocolToolAsyncAction is maintained as a legacy
+ *             fallback and will be removed in a future release.
+ */
+UCLASS(Abstract, MinimalAPI, meta=(DeprecatedNode, DeprecationMessage="Use UToolsetDefinition (ToolsetRegistry plugin) instead."))
+class UModelContextProtocolToolAsyncAction : public UCancellableAsyncAction
+{
+	GENERATED_BODY()
+
+public:
+	UE_API virtual void SetReadyToDestroy() override;
+
+	virtual FName GetToolFunctionName() const { return NAME_None; }
+	virtual FName GetToolResultPropertyName() const { return NAME_None; }
+
+	const FModelContextProtocolFunctionMetaData* GetToolFunctionMetaData() const { return &ToolFunctionMetaData; }
+	const FModelContextProtocolFunctionMetaData* GetToolResultMetaData() const { return &ToolResultMetaData; }
+
+	bool ShouldAutoRegisterTool() const { return bAutoRegisterTool; }
+
+	UE_API void RegisterTool();
+	UE_API void DeregisterTool();
+
+	//~ Begin UObject API
+	UE_API virtual void PostInitProperties() override;
+	UE_API virtual void FinishDestroy() override;
+	//~ End UObject API
+
+	FOnModelContextProtocolAsyncToolComplete OnAsyncToolComplete;
+
+	FModelContextProtocolToolRequestId ToolRequestId;
+
+protected:
+	/**
+	 * If true, automatically registers on module load.
+	 * If false, RegisterTool must manually be called.
+	 */ 
+	UPROPERTY(EditDefaultsOnly, Category = "Model Context Protocol")
+	bool bAutoRegisterTool = true;
+
+	UPROPERTY(VisibleDefaultsOnly, Category = "Model Context Protocol")
+	FModelContextProtocolFunctionMetaData ToolFunctionMetaData;
+
+	UPROPERTY(VisibleDefaultsOnly, Category = "Model Context Protocol")
+	FModelContextProtocolFunctionMetaData ToolResultMetaData;
+
+	TSharedPtr<FModelContextProtocolAsyncActionTool> Tool;
+
+	IModelContextProtocolTool::FResultCallback OnComplete;
+
+#if WITH_EDITORONLY_DATA
+	void CollectToolMetaData();
+#endif
+};
+
+/**
+ * IModelContextProtocolTool implementation, automatically derived from reflected UModelContextProtocolToolAsyncAction, using UFUNCTION comment
+ * to derived tool & parameter descriptions.
+ * @see IModelContextProtocolTool
+ * @see UModelContextProtocolToolAsyncAction
+ */
+struct FModelContextProtocolAsyncActionTool: public IModelContextProtocolTool
+{
+public:
+	FModelContextProtocolAsyncActionTool(UModelContextProtocolToolAsyncAction* InAsyncAction, UFunction* InFunction);
+
+	//~ Begin IModelContextProtocolTool API
+	virtual FString GetName() const override;
+	virtual FString GetDescription() const override;
+	virtual TSharedPtr<FJsonObject> GetInputJsonSchema() const override;
+	virtual TSharedPtr<FJsonObject> GetOutputJsonSchema() const override;
+	virtual void RunAsync(const FModelContextProtocolToolRequestId& RequestId, const TSharedPtr<FJsonObject>& Params, const FResultCallback& OnComplete) override;
+	virtual void CancelAsync(const FModelContextProtocolToolRequestId& RequestId) override;
+	virtual void AddReferencedObjects(FReferenceCollector& Collector) override;
+	//~ End IModelContextProtocolTool API
+
+protected:
+	// UModelContextProtocolToolAsyncAction subclass CDO
+	TObjectPtr<UModelContextProtocolToolAsyncAction> AsyncAction;
+	TObjectPtr<UFunction> Function;
+	TArray<TObjectPtr<UModelContextProtocolToolAsyncAction>> InProgressActions;
+
+	mutable TSharedPtr<FJsonObject> InputJsonSchema;
+	mutable TOptional<EModelContextProtocolToolResultType> ResultType;
+	mutable TSharedPtr<FJsonObject> OutputJsonSchema;
+	FStructOnScope FunctionParamsContainer;
+
+	void CacheResultTypeInfo() const;
+};
+
+#undef UE_API
diff --git a/Engine/Plugins/Experimental/ModelContextProtocol/Source/ModelContextProtocolEngine/Public/ModelContextProtocolToolLibrary.h b/Engine/Plugins/Experimental/ModelContextProtocol/Source/ModelContextProtocolEngine/Public/ModelContextProtocolToolLibrary.h
new file mode 100644
--- /dev/null
+++ b/Engine/Plugins/Experimental/ModelContextProtocol/Source/ModelContextProtocolEngine/Public/ModelContextProtocolToolLibrary.h
@@ -0,0 +1,126 @@
+// Copyright Epic Games, Inc. All Rights Reserved.
+
+#pragma once
+
+#include "IModelContextProtocolTool.h"
+
+#include "Engine/Blueprint.h"
+#include "ModelContextProtocolMetaData.h"
+#include "Kismet/BlueprintFunctionLibrary.h"
+#include "UObject/StructOnScope.h"
+
+#include "ModelContextProtocolToolLibrary.generated.h"
+
+#define UE_API MODELCONTEXTPROTOCOLENGINE_API
+
+class UModelContextProtocolToolLibrary;
+struct FModelContextProtocolLibraryTool;
+
+/**
+ * Automatically registers an FModelContextProtocolLibraryTool for each public UFUNCTION on module load, using cached / cooked meta-data.
+ *
+ * UBlueprintFunctionLibrary specialization allows caching and cooking otherwise-editor-only UFUNCTION descriptions and parameter meta-data.
+ *
+ * Can be subclassed in either C++ or Blueprints.
+ *
+ * For Blueprints:	Create via Content Browser -> Add -> MCP Tool Library, then simply define public functions with a doxygen-style
+ *					function tooltip.
+ *
+ * @deprecated Use UToolsetDefinition (ToolsetRegistry plugin) instead. UModelContextProtocolToolLibrary is maintained as a legacy fallback
+ *             and will be removed in a future release.
+ */
+UCLASS(Abstract, MinimalAPI, meta=(DeprecatedNode, DeprecationMessage="Use UToolsetDefinition (ToolsetRegistry plugin) instead."))
+class UModelContextProtocolToolLibrary : public UBlueprintFunctionLibrary
+{
+	GENERATED_BODY()
+public:
+	const TMap<FName, FModelContextProtocolFunctionMetaData>& GetFunctionMetaData() const
+	{
+		return FunctionMetaData;
+	}
+
+	const FModelContextProtocolFunctionMetaData* FindFunctionMetaData(const FName& FunctionName) const
+	{
+		return FunctionMetaData.Find(FunctionName);
+	}
+
+	bool ShouldAutoRegisterTools() const { return bAutoRegisterTools; }
+
+	UE_API void RegisterTools();
+	UE_API void DeregisterTools();
+	
+	//~ Begin UObject API
+	UE_API virtual void PostInitProperties() override;
+	UE_API virtual void PostLoad() override;
+	UE_API virtual void FinishDestroy() override;
+	//~ End UObject API
+	
+protected:
+
+	/**
+	 * If true, automatically registers an FModelContextProtocolLibraryTool for each public UFUNCTION on module load.
+	 * If false, RegisterTools must manually be called.
+	 */ 
+	UPROPERTY(EditDefaultsOnly, Category = "Model Context Protocol")
+	bool bAutoRegisterTools = true;
+
+	/** Per-UFUNCTION cached meta-data to allow meta-data use on cooked platforms */ 
+	UPROPERTY(VisibleDefaultsOnly, Category = "Model Context Protocol")
+	TMap<FName, FModelContextProtocolFunctionMetaData> FunctionMetaData;
+
+	TArray<TSharedRef<FModelContextProtocolLibraryTool>> Tools;
+
+#if WITH_EDITOR
+	UE_API virtual void PostCDOCompiled(const FPostCDOCompiledContext& Context) override;
+#endif
+#if WITH_EDITORONLY_DATA
+	void CollectFunctionMetaData();
+#endif
+};
+
+UCLASS(MinimalAPI)
+class UModelContextProtocolToolLibraryBlueprint : public UBlueprint
+{
+	GENERATED_BODY()
+public:
+	UModelContextProtocolToolLibraryBlueprint();
+
+	//~ Begin UBlueprint API
+#if WITH_EDITOR
+	virtual bool AlwaysCompileOnLoad() const override;
+#endif // WITH_EDITOR
+	//~ End UBlueprint API
+};
+
+/**
+ * IModelContextProtocolTool implementation, automatically derived from reflected UModelContextProtocolToolLibrary functions, using UFUNCTION comment
+ * to derived tool & parameter descriptions.
+ * @see IModelContextProtocolTool
+ * @see UModelContextProtocolToolLibrary
+ */
+struct FModelContextProtocolLibraryTool : public IModelContextProtocolTool
+{
+public:
+	FModelContextProtocolLibraryTool(UModelContextProtocolToolLibrary* InLibrary, UFunction* InFunction);
+	
+	//~ Begin IModelContextProtocolTool API
+	virtual FString GetName() const override;
+	virtual FString GetDescription() const override; 
+	virtual TSharedPtr<FJsonObject> GetInputJsonSchema() const override;
+	virtual TSharedPtr<FJsonObject> GetOutputJsonSchema() const override;
+	virtual FModelContextProtocolToolResult Run(const TSharedPtr<FJsonObject>& Params) override;
+	virtual void AddReferencedObjects(FReferenceCollector& Collector) override;
+	//~ End IModelContextProtocolTool API
+protected:
+	// UModelContextProtocolToolLibrary subclass CDO
+	TObjectPtr<UModelContextProtocolToolLibrary> Library;
+	TObjectPtr<UFunction> Function;
+	mutable TSharedPtr<FJsonObject> InputJsonSchema;
+	mutable TOptional<EModelContextProtocolToolResultType> ResultType;
+	mutable TSharedPtr<FJsonObject> OutputJsonSchema;
+	FStructOnScope FunctionParamsContainer;
+
+	void CacheResultTypeInfo() const;
+};
+
+#undef UE_API
diff --git a/Engine/Plugins/Experimental/ModelContextProtocol/Source/ModelContextProtocolEngine/Public/ModelContextProtocolToolUtils.h b/Engine/Plugins/Experimental/ModelContextProtocol/Source/ModelContextProtocolEngine/Public/ModelContextProtocolToolUtils.h
new file mode 100644
--- /dev/null
+++ b/Engine/Plugins/Experimental/ModelContextProtocol/Source/ModelContextProtocolEngine/Public/ModelContextProtocolToolUtils.h
@@ -0,0 +1,18 @@
+// Copyright Epic Games, Inc. All Rights Reserved.
+
+#pragma once
+
+#include "ModelContextProtocolToolResults.h"
+
+class FProperty;
+struct FModelContextProtocolFunctionMetaData;
+
+#define UE_API MODELCONTEXTPROTOCOLENGINE_API
+
+namespace UE::ModelContextProtocol
+{
+	UE_API EModelContextProtocolToolResultType GetToolResultType(FProperty* Property, const FModelContextProtocolFunctionMetaData* MetaData, TSharedPtr<FJsonObject>& OutSchema);
+	UE_API FModelContextProtocolToolResult GetToolResultFromType(EModelContextProtocolToolResultType ResultType, FProperty* Property, void* Container, uint16 ReturnValueOffset);
+};
+
+#undef UE_API
\ No newline at end of file
diff --git a/Engine/Plugins/Experimental/ModelContextProtocol/Source/ModelContextProtocolTests/ModelContextProtocolTests.Build.cs b/Engine/Plugins/Experimental/ModelContextProtocol/Source/ModelContextProtocolTests/ModelContextProtocolTests.Build.cs
--- a/Engine/Plugins/Experimental/ModelContextProtocol/Source/ModelContextProtocolTests/ModelContextProtocolTests.Build.cs
+++ b/Engine/Plugins/Experimental/ModelContextProtocol/Source/ModelContextProtocolTests/ModelContextProtocolTests.Build.cs
@@ -18,22 +18,11 @@ public class ModelContextProtocolTests : ModuleRules
 			new string[]
 			{
 				"CoreUObject",
-				"Engine",
 				"Json",
 				"JsonUtilities",
 				"HTTPServer",
 				"HTTP",
 				"ModelContextProtocol"
 			});
-
-		if (Target.bBuildEditor)
-		{
-			PrivateDependencyModuleNames.AddRange(
-				new string[]
-				{
-					"UnrealEd",
-					"ToolsetRegistry"
-				});
-		}
 	}
 }
diff --git a/Engine/Plugins/Experimental/ModelContextProtocol/Source/ModelContextProtocolTests/Private/Mocks/MockToolsetDefinition.cpp b/Engine/Plugins/Experimental/ModelContextProtocol/Source/ModelContextProtocolTests/Private/Mocks/MockToolsetDefinition.cpp
deleted file mode 100644
diff --git a/Engine/Plugins/Experimental/ModelContextProtocol/Source/ModelContextProtocolTests/Private/Mocks/MockToolsetDefinition.h b/Engine/Plugins/Experimental/ModelContextProtocol/Source/ModelContextProtocolTests/Private/Mocks/MockToolsetDefinition.h
deleted file mode 100644
diff --git a/Engine/Plugins/Experimental/ModelContextProtocol/Source/ModelContextProtocolTests/Private/ModelContextProtocolTestUtilities.h b/Engine/Plugins/Experimental/ModelContextProtocol/Source/ModelContextProtocolTests/Private/ModelContextProtocolTestUtilities.h
--- a/Engine/Plugins/Experimental/ModelContextProtocol/Source/ModelContextProtocolTests/Private/ModelContextProtocolTestUtilities.h
+++ b/Engine/Plugins/Experimental/ModelContextProtocol/Source/ModelContextProtocolTests/Private/ModelContextProtocolTestUtilities.h
@@ -3,7 +3,6 @@
 #pragma once
 
 #include "ModelContextProtocol.h"
-#include "ModelContextProtocolSettings.h"
 
 #include "Dom/JsonObject.h"
 #include "Dom/JsonValue.h"
@@ -17,8 +16,7 @@ namespace UE::ModelContextProtocol::Tests
 {
 	inline FString GetTestBaseUrl()
 	{
-		const UModelContextProtocolSettings* Settings = GetDefault<UModelContextProtocolSettings>();
-		return FString::Printf(TEXT("http://127.0.0.1:%d%s"), Settings->ServerPortNumber, *Settings->ServerUrlPath);
+		return FString::Printf(TEXT("http://127.0.0.1:%d%s"), UE::ModelContextProtocol::DefaultServerPort, UE::ModelContextProtocol::DefaultServerUrlPath);
 	}
 
 	inline FString JsonObjectToString(const TSharedRef<FJsonObject>& JsonObject)
diff --git a/Engine/Plugins/Experimental/ModelContextProtocol/Source/ModelContextProtocolTests/Private/Tests/ModelContextProtocolEngineSubsystemTests.cpp b/Engine/Plugins/Experimental/ModelContextProtocol/Source/ModelContextProtocolTests/Private/Tests/ModelContextProtocolEngineSubsystemTests.cpp
--- a/Engine/Plugins/Experimental/ModelContextProtocol/Source/ModelContextProtocolTests/Private/Tests/ModelContextProtocolEngineSubsystemTests.cpp
+++ b/Engine/Plugins/Experimental/ModelContextProtocol/Source/ModelContextProtocolTests/Private/Tests/ModelContextProtocolEngineSubsystemTests.cpp
@@ -1,6 +1,5 @@
 // Copyright Epic Games, Inc. All Rights Reserved.
 
-#include "ModelContextProtocolEngineSubsystem.h"
 #include "IModelContextProtocolModule.h"
 #include "IModelContextProtocolTool.h"
 #include "IModelContextProtocolResourceProvider.h"
@@ -12,7 +11,6 @@
 
 #include "Dom/JsonObject.h"
 #include "Dom/JsonValue.h"
-#include "Engine/Engine.h"
 #include "HttpModule.h"
 #include "Interfaces/IHttpRequest.h"
 #include "Interfaces/IHttpResponse.h"
@@ -21,8 +19,7 @@
 
 #if WITH_DEV_AUTOMATION_TESTS
 
-BEGIN_DEFINE_SPEC(FModelContextProtocolEngineSubsystemTests, "AI.ModelContextProtocol.EngineSubsystem", EAutomationTestFlags::ProductFilter | EAutomationTestFlags_ApplicationContextMask)
-	UModelContextProtocolEngineSubsystem* Subsystem;
+BEGIN_DEFINE_SPEC(FModelContextProtocolServerTests, "AI.ModelContextProtocol.Server", EAutomationTestFlags::ProductFilter | EAutomationTestFlags_ApplicationContextMask)
 	FString SessionId;
 	TSharedPtr<FMockModelContextProtocolTool> MockTool;
 	TArray<TSharedPtr<FMockModelContextProtocolTool>> AdditionalMockTools;
@@ -112,15 +109,14 @@ BEGIN_DEFINE_SPEC(FModelContextProtocolEngineSubsystemTests, "AI.ModelContextPro
 		});
 		DeleteRequest->ProcessRequest();
 	}
-END_DEFINE_SPEC(FModelContextProtocolEngineSubsystemTests)
+END_DEFINE_SPEC(FModelContextProtocolServerTests)
 
-void FModelContextProtocolEngineSubsystemTests::Define()
+void FModelContextProtocolServerTests::Define()
 {
 	using namespace UE::ModelContextProtocol::Tests;
 
 	BeforeEach([this]()
 	{
-		Subsystem = GEngine ? GEngine->GetEngineSubsystem<UModelContextProtocolEngineSubsystem>() : nullptr;
 		MockTool = MakeShared<FMockModelContextProtocolTool>();
 		MockResourceProvider = MakeShared<FMockModelContextProtocolResourceProvider>();
 		SessionId.Reset();
diff --git a/Engine/Plugins/Experimental/ModelContextProtocol/Source/ModelContextProtocolTests/Private/Tests/ModelContextProtocolToolNameTests.cpp b/Engine/Plugins/Experimental/ModelContextProtocol/Source/ModelContextProtocolTests/Private/Tests/ModelContextProtocolToolNameTests.cpp
new file mode 100644
--- /dev/null
+++ b/Engine/Plugins/Experimental/ModelContextProtocol/Source/ModelContextProtocolTests/Private/Tests/ModelContextProtocolToolNameTests.cpp
@@ -0,0 +1,150 @@
+// Copyright Epic Games, Inc. All Rights Reserved.
+
+#include "ModelContextProtocol.h"
+#include "IModelContextProtocolModule.h"
+#include "Mocks/MockModelContextProtocolTool.h"
+
+#include "Misc/AutomationTest.h"
+
+#if WITH_DEV_AUTOMATION_TESTS
+
+BEGIN_DEFINE_SPEC(FModelContextProtocolToolNameTests, "AI.ModelContextProtocol.ToolNames", EAutomationTestFlags::ProductFilter | EAutomationTestFlags_ApplicationContextMask)
+END_DEFINE_SPEC(FModelContextProtocolToolNameTests)
+
+void FModelContextProtocolToolNameTests::Define()
+{
+	using namespace UE::ModelContextProtocol;
+
+	Describe("ValidateToolName", [this]()
+	{
+		It("should accept simple alphanumeric names", [this]()
+		{
+			TestEqual("getUser", ValidateToolName(TEXT("getUser")), EToolNameValidation::Valid);
+			TestEqual("DATA_EXPORT_v2", ValidateToolName(TEXT("DATA_EXPORT_v2")), EToolNameValidation::Valid);
+		});
+
+		It("should accept names with dots", [this]()
+		{
+			TestEqual("admin.tools.list", ValidateToolName(TEXT("admin.tools.list")), EToolNameValidation::Valid);
+			TestEqual("dotted PascalCase", ValidateToolName(TEXT("ToolsetRegistry.EditorAppToolset.GetSelectedActors")), EToolNameValidation::Valid);
+		});
+
+		It("should accept names with hyphens and underscores", [this]()
+		{
+			TestEqual("my-tool", ValidateToolName(TEXT("my-tool")), EToolNameValidation::Valid);
+			TestEqual("my_tool", ValidateToolName(TEXT("my_tool")), EToolNameValidation::Valid);
+			TestEqual("my-tool_v2.0", ValidateToolName(TEXT("my-tool_v2.0")), EToolNameValidation::Valid);
+		});
+
+		It("should accept names with snake_case and dots", [this]()
+		{
+			TestEqual("snake_case.dotted", ValidateToolName(TEXT("toolset_registry.toolsets.core.actor.ActorTools.get_label")), EToolNameValidation::Valid);
+		});
+
+		It("should reject empty names", [this]()
+		{
+			TestEqual("empty string", ValidateToolName(TEXT("")), EToolNameValidation::Empty);
+		});
+
+		It("should reject names exceeding 128 characters", [this]()
+		{
+			FString LongName;
+			for (int32 Index = 0; Index < 129; ++Index)
+			{
+				LongName += TEXT("a");
+			}
+			TestEqual("129 chars", ValidateToolName(LongName), EToolNameValidation::ExceedsMaxLength);
+		});
+
+		It("should accept names at exactly 128 characters", [this]()
+		{
+			FString MaxName;
+			for (int32 Index = 0; Index < 128; ++Index)
+			{
+				MaxName += TEXT("a");
+			}
+			TestEqual("128 chars", ValidateToolName(MaxName), EToolNameValidation::Valid);
+		});
+
+		It("should reject names with spaces", [this]()
+		{
+			TestEqual("my tool", ValidateToolName(TEXT("my tool")), EToolNameValidation::InvalidCharacters);
+		});
+
+		It("should reject names with special characters", [this]()
+		{
+			TestEqual("my@tool", ValidateToolName(TEXT("my@tool")), EToolNameValidation::InvalidCharacters);
+			TestEqual("my/tool", ValidateToolName(TEXT("my/tool")), EToolNameValidation::InvalidCharacters);
+			TestEqual("my:tool", ValidateToolName(TEXT("my:tool")), EToolNameValidation::InvalidCharacters);
+			TestEqual("my,tool", ValidateToolName(TEXT("my,tool")), EToolNameValidation::InvalidCharacters);
+		});
+	});
+
+	Describe("AddTool name validation", [this]()
+	{
+		It("should reject tools with empty names", [this]()
+		{
+			TSharedPtr<FMockModelContextProtocolTool> MockTool = MakeShared<FMockModelContextProtocolTool>();
+			MockTool->Name = TEXT("");
+
+			AddExpectedError(TEXT("Rejecting tool with empty name"), EAutomationExpectedErrorFlags::Contains);
+
+			IModelContextProtocolModule& Module = IModelContextProtocolModule::GetChecked();
+			TestFalse("Empty-named tool should be rejected", Module.AddTool(MockTool.ToSharedRef()));
+		});
+
+		It("should still register tools with invalid characters", [this]()
+		{
+			TSharedPtr<FMockModelContextProtocolTool> MockTool = MakeShared<FMockModelContextProtocolTool>();
+			MockTool->Name = TEXT("my invalid tool name");
+
+			IModelContextProtocolModule& Module = IModelContextProtocolModule::GetChecked();
+			const bool bAdded = Module.AddTool(MockTool.ToSharedRef());
+			TestTrue("Tool should still be registered despite non-compliant name", bAdded);
+
+			if (bAdded)
+			{
+				Module.RemoveTool(MockTool.ToSharedRef());
+			}
+		});
+
+		It("should warn for tool names exceeding 128 characters", [this]()
+		{
+			TSharedPtr<FMockModelContextProtocolTool> MockTool = MakeShared<FMockModelContextProtocolTool>();
+			FString LongName;
+			for (int32 Index = 0; Index < 129; ++Index)
+			{
+				LongName += TEXT("a");
+			}
+			MockTool->Name = LongName;
+
+			AddExpectedError(TEXT("exceeds MCP spec limit"), EAutomationExpectedErrorFlags::Contains);
+
+			IModelContextProtocolModule& Module = IModelContextProtocolModule::GetChecked();
+			const bool bAdded = Module.AddTool(MockTool.ToSharedRef());
+			TestTrue("Tool should still be registered despite length violation", bAdded);
+
+			if (bAdded)
+			{
+				Module.RemoveTool(MockTool.ToSharedRef());
+			}
+		});
+
+		It("should register tools with compliant names without warning", [this]()
+		{
+			TSharedPtr<FMockModelContextProtocolTool> MockTool = MakeShared<FMockModelContextProtocolTool>();
+			MockTool->Name = TEXT("valid.tool-name_v2");
+
+			IModelContextProtocolModule& Module = IModelContextProtocolModule::GetChecked();
+			const bool bAdded = Module.AddTool(MockTool.ToSharedRef());
+			TestTrue("Tool should be registered", bAdded);
+
+			if (bAdded)
+			{
+				Module.RemoveTool(MockTool.ToSharedRef());
+			}
+		});
+	});
+}
+
+#endif // WITH_DEV_AUTOMATION_TESTS
diff --git a/Engine/Plugins/Experimental/ModelContextProtocol/Source/ModelContextProtocolTests/Private/Tests/ModelContextProtocolToolUtilsTests.cpp b/Engine/Plugins/Experimental/ModelContextProtocol/Source/ModelContextProtocolTests/Private/Tests/ModelContextProtocolToolUtilsTests.cpp
deleted file mode 100644
diff --git a/Engine/Plugins/Experimental/ModelContextProtocol/Source/ModelContextProtocolTests/Private/Tests/ModelContextProtocolToolsetRegistryTests.cpp b/Engine/Plugins/Experimental/ModelContextProtocol/Source/ModelContextProtocolTests/Private/Tests/ModelContextProtocolToolsetRegistryTests.cpp
deleted file mode 100644
