[ModelContextProtocol] Ensure `ModelContextProtocolEditor` startup compatibility across branches.

- Replace `FCoreDelegates::GetOnPostEngineInit()` with the `FOnPostEngineInit<FCoreDelegates>` compatibility shim from `ToolsetRegistry`, which gracefully falls back to the deprecated `FCoreDelegates::OnPostEngineInit` on older branches that lack the new API.
- Use `FOnPostEngineInit<FCoreDelegates>::Get()` instead of `FCoreDelegates::GetOnPostEngineInit()`.
- Add `GEditor` early-out: if editor is already available at `StartupModule` time, run setup immediately without deferring.
- Adopt `FDelegateHandleRaii` from ToolsetRegistry for automatic delegate cleanup, matching the AIAssistant pattern.

#Synced-CL 52520731

#Synced-CL 52520731

diff --git a/Engine/Plugins/Experimental/ModelContextProtocol/Source/ModelContextProtocolEditor/Private/ModelContextProtocolEditor.cpp b/Engine/Plugins/Experimental/ModelContextProtocol/Source/ModelContextProtocolEditor/Private/ModelContextProtocolEditor.cpp
--- a/Engine/Plugins/Experimental/ModelContextProtocol/Source/ModelContextProtocolEditor/Private/ModelContextProtocolEditor.cpp
+++ b/Engine/Plugins/Experimental/ModelContextProtocol/Source/ModelContextProtocolEditor/Private/ModelContextProtocolEditor.cpp
@@ -8,7 +8,7 @@
 
 #include "Editor.h"
 #include "KismetCompilerModule.h"
-#include "Misc/CoreDelegates.h"
+#include "ToolsetRegistry/Private/ToolsetRegistry/EngineDelegatesCompatibility.h"
 #include "ToolsetRegistry/ToolsetRegistry.h"
 #include "ToolsetRegistry/ToolsetRegistrySubsystem.h"
 
@@ -29,36 +29,50 @@ void FModelContextProtocolEditorModule::StartupModule()
 		});
 	}
 
-	// Defer ToolsetRegistry binding and auto-start until the editor is fully initialized
-	// (GEditor and editor subsystems are not available during StartupModule)
-	PostEngineInitHandle = FCoreDelegates::GetOnPostEngineInit().AddLambda([this]()
+	// If GEditor is already available, set up immediately. Otherwise, defer until post-engine-init.
+	// Uses FOnPostEngineInit compatibility shim to support branches that lack FCoreDelegates::GetOnPostEngineInit().
+	if (GEditor)
 	{
-		// Subscribe to ToolsetRegistry changes so adapters are re-registered when toolsets are added
-		if (UToolsetRegistrySubsystem* ToolsetRegistrySubsystem = GEditor ? GEditor->GetEditorSubsystem<UToolsetRegistrySubsystem>() : nullptr)
-		{
-			OnToolsetRegisteredHandle = ToolsetRegistrySubsystem->ToolsetRegistry.OnToolsetRegistered().AddLambda([this]()
+		SetupEditorIntegration();
+	}
+	else
+	{
+		FSimpleMulticastDelegate& OnPostEngineInit = FOnPostEngineInit<FCoreDelegates>::Get();
+		PostEngineInitHandle = UE::ToolsetRegistry::FDelegateHandleRaii::Create(
+			OnPostEngineInit, OnPostEngineInit.AddLambda([this]()
 			{
-				ToolsetRegistryAdapterManager.RegisterTools();
-			});
-		}
+				SetupEditorIntegration();
+			}));
+	}
+}
 
-		// Initial registration of any toolsets already loaded
-		ToolsetRegistryAdapterManager.RegisterTools();
+void FModelContextProtocolEditorModule::SetupEditorIntegration()
+{
+	// Subscribe to ToolsetRegistry changes so adapters are re-registered when toolsets are added
+	if (UToolsetRegistrySubsystem* ToolsetRegistrySubsystem = GEditor ? GEditor->GetEditorSubsystem<UToolsetRegistrySubsystem>() : nullptr)
+	{
+		OnToolsetRegisteredHandle = ToolsetRegistrySubsystem->ToolsetRegistry.OnToolsetRegistered().AddLambda([this]()
+		{
+			ToolsetRegistryAdapterManager.RegisterTools();
+		});
+	}
 
-		// Auto-start the MCP server if configured
-		if (UE::ModelContextProtocol::ShouldAutoStartServer())
+	// Initial registration of any toolsets already loaded
+	ToolsetRegistryAdapterManager.RegisterTools();
+
+	// Auto-start the MCP server if configured
+	if (UE::ModelContextProtocol::ShouldAutoStartServer())
+	{
+		if (IModelContextProtocolModule* Module = IModelContextProtocolModule::Get())
 		{
-			if (IModelContextProtocolModule* Module = IModelContextProtocolModule::Get())
-			{
-				Module->StartServer(UE::ModelContextProtocol::GetServerPortNumber(), UE::ModelContextProtocol::GetServerUrlPath());
-			}
+			Module->StartServer(UE::ModelContextProtocol::GetServerPortNumber(), UE::ModelContextProtocol::GetServerUrlPath());
 		}
-	});
+	}
 }
 
 void FModelContextProtocolEditorModule::ShutdownModule()
 {
-	FCoreDelegates::GetOnPostEngineInit().Remove(PostEngineInitHandle);
+	PostEngineInitHandle.Reset();
 
 	if (IModelContextProtocolModule* Module = IModelContextProtocolModule::Get())
 	{
diff --git a/Engine/Plugins/Experimental/ModelContextProtocol/Source/ModelContextProtocolEditor/Public/ModelContextProtocolEditor.h b/Engine/Plugins/Experimental/ModelContextProtocol/Source/ModelContextProtocolEditor/Public/ModelContextProtocolEditor.h
--- a/Engine/Plugins/Experimental/ModelContextProtocol/Source/ModelContextProtocolEditor/Public/ModelContextProtocolEditor.h
+++ b/Engine/Plugins/Experimental/ModelContextProtocol/Source/ModelContextProtocolEditor/Public/ModelContextProtocolEditor.h
@@ -5,6 +5,7 @@
 #include "ModelContextProtocolToolsetRegistryAdapter.h"
 
 #include "Modules/ModuleManager.h"
+#include "ToolsetRegistry/DelegateHandle.h"
 
 class FModelContextProtocolEditorModule : public IModuleInterface
 {
@@ -15,8 +16,10 @@ public:
 	virtual void ShutdownModule() override;
 
 private:
+	void SetupEditorIntegration();
+
 	FToolsetRegistryToolAdapterManager ToolsetRegistryAdapterManager;
-	FDelegateHandle PostEngineInitHandle;
+	UE::ToolsetRegistry::FDelegateHandleRaii PostEngineInitHandle;
 	FDelegateHandle OnRefreshToolsHandle;
 	FDelegateHandle OnToolsetRegisteredHandle;
 };
