[ModelContextProtocol] Add console command to generate *MCP* client configuration files.

- Add `ModelContextProtocol.GenerateClientConfig` console command that writes *MCP* server configuration files for supported AI coding clients.
- *JSON* formats are created or updated, preserving existing entries.
- *TOML* format (Codex) is written once; logs error if file already exists.
- Port and URL path read from `UModelContextProtocolSettings`.
- Add `EModelContextProtocolClient` enum and `WriteClientConfiguration()` API in `ModelContextProtocolEngine`.
- Add `AI.ModelContextProtocol.ClientConfig` test spec.
- Config files written to `FPaths::RootDir()` (workspace root) in source builds, `FPaths::ProjectDir()` in installed/launcher builds.

**Supported Clients**

- **Claude Code**: `.mcp.json` (`mcpServers` root, `url` + `type: "http"`)
- **Cursor**: `.cursor/mcp.json` (`mcpServers` root, `url`)
- **VS Code / Copilot**: `.vscode/mcp.json` (`servers` root, `url` + `type: "http"`)
- **Gemini CLI**: `.gemini/settings.json` (`mcpServers` root, `httpUrl`)
- **Codex CLI**: `.codex/config.toml` (*TOML* `[mcp_servers]` section, write-once)

Usage: `ModelContextProtocol.GenerateClientConfig <ClaudeCode|Cursor|VSCode|Gemini|Codex|All>`

#Synced-CL 52520731

diff --git a/Engine/Plugins/Experimental/ModelContextProtocol/ModelContextProtocol.uplugin b/Engine/Plugins/Experimental/ModelContextProtocol/ModelContextProtocol.uplugin
--- a/Engine/Plugins/Experimental/ModelContextProtocol/ModelContextProtocol.uplugin
+++ b/Engine/Plugins/Experimental/ModelContextProtocol/ModelContextProtocol.uplugin
@@ -37,6 +37,11 @@
 			"Type": "UncookedOnly",
 			"LoadingPhase": "Default"
 		},
+		{
+			"Name": "ModelContextProtocolEngineTests",
+			"Type": "UncookedOnly",
+			"LoadingPhase": "Default"
+		},
 		{
 			"Name": "ModelContextProtocolEditorTests",
 			"Type": "Editor",
diff --git a/Engine/Plugins/Experimental/ModelContextProtocol/Source/ModelContextProtocol/Public/ModelContextProtocol.h b/Engine/Plugins/Experimental/ModelContextProtocol/Source/ModelContextProtocol/Public/ModelContextProtocol.h
--- a/Engine/Plugins/Experimental/ModelContextProtocol/Source/ModelContextProtocol/Public/ModelContextProtocol.h
+++ b/Engine/Plugins/Experimental/ModelContextProtocol/Source/ModelContextProtocol/Public/ModelContextProtocol.h
@@ -44,6 +44,7 @@ namespace UE::ModelContextProtocol
 
 	static constexpr uint32 DefaultServerPort = 8000;
 	static constexpr const TCHAR* DefaultServerUrlPath = TEXT("/mcp");
+	static constexpr const TCHAR* DefaultServerName = TEXT("unreal-mcp");
 
 	/** Tool name validation result per MCP spec (2025-11-25). */
 	enum class EToolNameValidation : uint8
diff --git a/Engine/Plugins/Experimental/ModelContextProtocol/Source/ModelContextProtocolEngine/Private/ModelContextProtocolClientConfig.cpp b/Engine/Plugins/Experimental/ModelContextProtocol/Source/ModelContextProtocolEngine/Private/ModelContextProtocolClientConfig.cpp
new file mode 100644
--- /dev/null
+++ b/Engine/Plugins/Experimental/ModelContextProtocol/Source/ModelContextProtocolEngine/Private/ModelContextProtocolClientConfig.cpp
@@ -0,0 +1,190 @@
+// Copyright Epic Games, Inc. All Rights Reserved.
+
+#include "ModelContextProtocolClientConfig.h"
+#include "ModelContextProtocol.h"
+
+#include "Dom/JsonObject.h"
+#include "HAL/PlatformFileManager.h"
+#include "Misc/App.h"
+#include "Misc/FileHelper.h"
+#include "Misc/Paths.h"
+#include "Serialization/JsonReader.h"
+#include "Serialization/JsonSerializer.h"
+#include "Serialization/JsonWriter.h"
+
+namespace UE::ModelContextProtocol::Private
+{
+	struct FClientConfigDescriptor
+	{
+		FString RelativeFilePath;
+		FString ServersRootKey;
+		FString UrlFieldName;
+		bool bIncludeTypeField = false;
+		bool bIsToml = false;
+	};
+
+	FClientConfigDescriptor GetClientConfigDescriptor(EModelContextProtocolClient Client)
+	{
+		switch (Client)
+		{
+		case EModelContextProtocolClient::ClaudeCode:
+			return { TEXT(".mcp.json"), TEXT("mcpServers"), TEXT("url"), /*bIncludeTypeField*/ true };
+		case EModelContextProtocolClient::Cursor:
+			return { TEXT(".cursor/mcp.json"), TEXT("mcpServers"), TEXT("url") };
+		case EModelContextProtocolClient::VSCode:
+			return { TEXT(".vscode/mcp.json"), TEXT("servers"), TEXT("url"), /*bIncludeTypeField*/ true };
+		case EModelContextProtocolClient::Gemini:
+			return { TEXT(".gemini/settings.json"), TEXT("mcpServers"), TEXT("httpUrl") };
+		case EModelContextProtocolClient::Codex:
+			return { TEXT(".codex/config.toml"), TEXT("mcp_servers"), TEXT("url"), /*bIncludeTypeField*/ false, /*bIsToml*/ true };
+		default:
+			return {};
+		}
+	}
+
+	static constexpr const TCHAR* ServerEntryName = UE::ModelContextProtocol::DefaultServerName;
+
+	bool WriteJsonClientConfiguration(const FString& FilePath, const FClientConfigDescriptor& Descriptor, const FString& ServerUrl)
+	{
+		// Read existing file if present
+		TSharedPtr<FJsonObject> RootObject;
+		FString ExistingContent;
+		if (FFileHelper::LoadFileToString(ExistingContent, *FilePath))
+		{
+			TSharedRef<TJsonReader<>> Reader = TJsonReaderFactory<>::Create(ExistingContent);
+			if (!FJsonSerializer::Deserialize(Reader, RootObject) || !RootObject.IsValid())
+			{
+				UE_LOG(LogModelContextProtocol, Warning, TEXT("Existing file contains malformed JSON, overwriting: %s"), *FilePath);
+				RootObject.Reset();
+			}
+		}
+
+		if (!RootObject.IsValid())
+		{
+			RootObject = MakeShared<FJsonObject>();
+		}
+
+		// Get or create the servers object under the root key
+		const TSharedPtr<FJsonObject>* ExistingServersObject = nullptr;
+		TSharedPtr<FJsonObject> ServersObject;
+		if (RootObject->TryGetObjectField(Descriptor.ServersRootKey, ExistingServersObject) && ExistingServersObject)
+		{
+			ServersObject = *ExistingServersObject;
+		}
+		if (!ServersObject.IsValid())
+		{
+			ServersObject = MakeShared<FJsonObject>();
+		}
+
+		// Create the server entry
+		TSharedPtr<FJsonObject> ServerEntry = MakeShared<FJsonObject>();
+		if (Descriptor.bIncludeTypeField)
+		{
+			ServerEntry->SetStringField(TEXT("type"), TEXT("http"));
+		}
+		ServerEntry->SetStringField(Descriptor.UrlFieldName, ServerUrl);
+
+		ServersObject->SetObjectField(ServerEntryName, ServerEntry);
+		RootObject->SetObjectField(Descriptor.ServersRootKey, ServersObject);
+
+		// Write with pretty printing
+		FString OutputString;
+		TSharedRef<TJsonWriter<TCHAR, TPrettyJsonPrintPolicy<TCHAR>>> Writer =
+			TJsonWriterFactory<TCHAR, TPrettyJsonPrintPolicy<TCHAR>>::Create(&OutputString);
+		if (!FJsonSerializer::Serialize(RootObject.ToSharedRef(), Writer) || OutputString.IsEmpty())
+		{
+			UE_LOG(LogModelContextProtocol, Warning, TEXT("Failed to serialize JSON for: %s"), *FilePath);
+			return false;
+		}
+
+		// Ensure parent directory exists
+		const FString Directory = FPaths::GetPath(FilePath);
+		if (!FPlatformFileManager::Get().GetPlatformFile().CreateDirectoryTree(*Directory))
+		{
+			UE_LOG(LogModelContextProtocol, Warning, TEXT("Failed to create directory: %s"), *Directory);
+			return false;
+		}
+
+		if (FFileHelper::SaveStringToFile(OutputString, *FilePath))
+		{
+			UE_LOG(LogModelContextProtocol, Display, TEXT("MCP client configuration written to: %s"), *FilePath);
+			return true;
+		}
+
+		UE_LOG(LogModelContextProtocol, Warning, TEXT("Failed to write MCP client configuration to: %s"), *FilePath);
+		return false;
+	}
+
+	bool WriteTomlClientConfiguration(const FString& FilePath, const FString& ServerUrl)
+	{
+		if (FPlatformFileManager::Get().GetPlatformFile().FileExists(*FilePath))
+		{
+			UE_LOG(LogModelContextProtocol, Error, TEXT("Codex configuration already exists at: %s. TOML format does not support safe upsert — edit manually to add [mcp_servers.%s]."), *FilePath, ServerEntryName);
+			return false;
+		}
+
+		// Ensure parent directory exists
+		const FString Directory = FPaths::GetPath(FilePath);
+		if (!FPlatformFileManager::Get().GetPlatformFile().CreateDirectoryTree(*Directory))
+		{
+			UE_LOG(LogModelContextProtocol, Warning, TEXT("Failed to create directory: %s"), *Directory);
+			return false;
+		}
+
+		FString TomlContent;
+		TomlContent += FString::Printf(TEXT("[mcp_servers.%s]\n"), ServerEntryName);
+		TomlContent += FString::Printf(TEXT("url = \"%s\"\n"), *ServerUrl);
+
+		if (FFileHelper::SaveStringToFile(TomlContent, *FilePath))
+		{
+			UE_LOG(LogModelContextProtocol, Display, TEXT("MCP client configuration written to: %s"), *FilePath);
+			return true;
+		}
+
+		UE_LOG(LogModelContextProtocol, Warning, TEXT("Failed to write MCP client configuration to: %s"), *FilePath);
+		return false;
+	}
+}
+
+bool UE::ModelContextProtocol::WriteClientConfiguration(EModelContextProtocolClient Client, uint32 Port, const FString& UrlPath, const FString& BaseDirectory)
+{
+	const Private::FClientConfigDescriptor Descriptor = Private::GetClientConfigDescriptor(Client);
+	if (Descriptor.RelativeFilePath.IsEmpty())
+	{
+		UE_LOG(LogModelContextProtocol, Warning, TEXT("Unknown MCP client type"));
+		return false;
+	}
+
+	const FString ServerUrl = FString::Printf(TEXT("http://127.0.0.1:%u%s"), Port, *UrlPath);
+	// In source builds, FPaths::RootDir() returns the workspace root (the directory containing Engine/),
+	// which is where AI coding clients expect config files.
+	// In installed/launcher builds, RootDir() points to the engine installation, so fall back to ProjectDir().
+	const FString RootDirectory = BaseDirectory.IsEmpty()
+		? (FApp::IsEngineInstalled() ? FPaths::ProjectDir() : FPaths::RootDir())
+		: BaseDirectory;
+	const FString FilePath = FPaths::Combine(RootDirectory, Descriptor.RelativeFilePath);
+
+	if (Descriptor.bIsToml)
+	{
+		return Private::WriteTomlClientConfiguration(FilePath, ServerUrl);
+	}
+	else
+	{
+		return Private::WriteJsonClientConfiguration(FilePath, Descriptor, ServerUrl);
+	}
+}
+
+int32 UE::ModelContextProtocol::WriteAllClientConfigurations(uint32 Port, const FString& UrlPath, const FString& BaseDirectory)
+{
+	int32 SuccessCount = 0;
+
+	for (uint8 Index = 0; Index <= static_cast<uint8>(EModelContextProtocolClient::Codex); ++Index)
+	{
+		if (WriteClientConfiguration(static_cast<EModelContextProtocolClient>(Index), Port, UrlPath, BaseDirectory))
+		{
+			++SuccessCount;
+		}
+	}
+
+	return SuccessCount;
+}
diff --git a/Engine/Plugins/Experimental/ModelContextProtocol/Source/ModelContextProtocolEngine/Private/ModelContextProtocolEngineModule.cpp b/Engine/Plugins/Experimental/ModelContextProtocol/Source/ModelContextProtocolEngine/Private/ModelContextProtocolEngineModule.cpp
--- a/Engine/Plugins/Experimental/ModelContextProtocol/Source/ModelContextProtocolEngine/Private/ModelContextProtocolEngineModule.cpp
+++ b/Engine/Plugins/Experimental/ModelContextProtocol/Source/ModelContextProtocolEngine/Private/ModelContextProtocolEngineModule.cpp
@@ -1,7 +1,72 @@
 // Copyright Epic Games, Inc. All Rights Reserved.
 
+#include "ModelContextProtocol.h"
+#include "ModelContextProtocolClientConfig.h"
+#include "ModelContextProtocolSettings.h"
+
+#include "HAL/IConsoleManager.h"
 #include "Modules/ModuleManager.h"
 
+namespace UE::ModelContextProtocol::Private
+{
+	TOptional<EModelContextProtocolClient> ParseClientName(const FString& Name)
+	{
+		if (Name.Equals(TEXT("ClaudeCode"), ESearchCase::IgnoreCase) || Name.Equals(TEXT("Claude"), ESearchCase::IgnoreCase))
+		{
+			return EModelContextProtocolClient::ClaudeCode;
+		}
+		if (Name.Equals(TEXT("Cursor"), ESearchCase::IgnoreCase))
+		{
+			return EModelContextProtocolClient::Cursor;
+		}
+		if (Name.Equals(TEXT("VSCode"), ESearchCase::IgnoreCase) || Name.Equals(TEXT("Copilot"), ESearchCase::IgnoreCase))
+		{
+			return EModelContextProtocolClient::VSCode;
+		}
+		if (Name.Equals(TEXT("Gemini"), ESearchCase::IgnoreCase))
+		{
+			return EModelContextProtocolClient::Gemini;
+		}
+		if (Name.Equals(TEXT("Codex"), ESearchCase::IgnoreCase))
+		{
+			return EModelContextProtocolClient::Codex;
+		}
+		return {};
+	}
+
+	FAutoConsoleCommand CommandGenerateClientConfig = FAutoConsoleCommand(
+		TEXT("ModelContextProtocol.GenerateClientConfig"),
+		TEXT("Generates MCP client configuration file(s). Usage: ModelContextProtocol.GenerateClientConfig <ClaudeCode|Cursor|VSCode|Gemini|Codex|All>"),
+		FConsoleCommandWithArgsDelegate::CreateLambda([](const TArray<FString>& Args)
+		{
+			if (Args.IsEmpty())
+			{
+				UE_LOG(LogModelContextProtocol, Warning, TEXT("Usage: ModelContextProtocol.GenerateClientConfig <ClaudeCode|Cursor|VSCode|Gemini|Codex|All>"));
+				return;
+			}
+
+			const uint32 Port = GetServerPortNumber();
+			const FString UrlPath = GetServerUrlPath();
+
+			if (Args[0].Equals(TEXT("All"), ESearchCase::IgnoreCase))
+			{
+				const int32 Count = WriteAllClientConfigurations(Port, UrlPath);
+				UE_LOG(LogModelContextProtocol, Display, TEXT("Generated %d MCP client configuration file(s)"), Count);
+			}
+			else
+			{
+				TOptional<EModelContextProtocolClient> Client = ParseClientName(Args[0]);
+				if (!Client.IsSet())
+				{
+					UE_LOG(LogModelContextProtocol, Warning, TEXT("Unknown client \"%s\". Supported: ClaudeCode, Cursor, VSCode, Gemini, Codex, All"), *Args[0]);
+					return;
+				}
+				WriteClientConfiguration(Client.GetValue(), Port, UrlPath);
+			}
+		}),
+		ECVF_Cheat);
+}
+
 class FModelContextProtocolEngineModule : public IModuleInterface
 {
 };
diff --git a/Engine/Plugins/Experimental/ModelContextProtocol/Source/ModelContextProtocolEngine/Public/ModelContextProtocolClientConfig.h b/Engine/Plugins/Experimental/ModelContextProtocol/Source/ModelContextProtocolEngine/Public/ModelContextProtocolClientConfig.h
new file mode 100644
--- /dev/null
+++ b/Engine/Plugins/Experimental/ModelContextProtocol/Source/ModelContextProtocolEngine/Public/ModelContextProtocolClientConfig.h
@@ -0,0 +1,49 @@
+// Copyright Epic Games, Inc. All Rights Reserved.
+
+#pragma once
+
+#include "ModelContextProtocolClientConfig.generated.h"
+
+/** Supported AI coding clients for MCP configuration file generation. */
+UENUM()
+enum class EModelContextProtocolClient : uint8
+{
+	/** Claude Code: `.mcp.json` in project root */
+	ClaudeCode,
+	/** Cursor: `.cursor/mcp.json` in project root */
+	Cursor,
+	/** VS Code / Copilot: `.vscode/mcp.json` in project root */
+	VSCode,
+	/** Gemini CLI: `.gemini/settings.json` in project root */
+	Gemini,
+	/** Codex CLI: `.codex/config.toml` in project root (TOML, write-once) */
+	Codex
+};
+
+namespace UE::ModelContextProtocol
+{
+	/**
+	 * Generates an MCP client configuration file.
+	 *
+	 * JSON formats (ClaudeCode, Cursor, VSCode, Gemini) are created or updated, preserving existing entries.
+	 * TOML format (Codex) is written once; logs error and returns false if file already exists.
+	 *
+	 * @param Client The target AI client.
+	 * @param Port The MCP server port number.
+	 * @param UrlPath The MCP server URL path (e.g., "/mcp").
+	 * @param BaseDirectory The root directory for config files. Defaults to FPaths::RootDir() in source builds or FPaths::ProjectDir() in installed builds if empty.
+	 * @return true if the file was written successfully.
+	 * @see EModelContextProtocolClient
+	 */
+	MODELCONTEXTPROTOCOLENGINE_API bool WriteClientConfiguration(EModelContextProtocolClient Client, uint32 Port, const FString& UrlPath, const FString& BaseDirectory = FString());
+
+	/**
+	 * Generates MCP client configuration files for all supported clients.
+	 *
+	 * @param Port The MCP server port number.
+	 * @param UrlPath The MCP server URL path (e.g., "/mcp").
+	 * @param BaseDirectory The root directory for config files. Defaults to FPaths::RootDir() in source builds or FPaths::ProjectDir() in installed builds if empty.
+	 * @return Number of files successfully written.
+	 */
+	MODELCONTEXTPROTOCOLENGINE_API int32 WriteAllClientConfigurations(uint32 Port, const FString& UrlPath, const FString& BaseDirectory = FString());
+}
diff --git a/Engine/Plugins/Experimental/ModelContextProtocol/Source/ModelContextProtocolEngineTests/ModelContextProtocolEngineTests.Build.cs b/Engine/Plugins/Experimental/ModelContextProtocol/Source/ModelContextProtocolEngineTests/ModelContextProtocolEngineTests.Build.cs
new file mode 100644
--- /dev/null
+++ b/Engine/Plugins/Experimental/ModelContextProtocol/Source/ModelContextProtocolEngineTests/ModelContextProtocolEngineTests.Build.cs
@@ -0,0 +1,27 @@
+// Copyright Epic Games, Inc. All Rights Reserved.
+
+using UnrealBuildTool;
+
+public class ModelContextProtocolEngineTests : ModuleRules
+{
+	public ModelContextProtocolEngineTests(ReadOnlyTargetRules Target) : base(Target)
+	{
+		PCHUsage = ModuleRules.PCHUsageMode.UseExplicitOrSharedPCHs;
+
+		PublicDependencyModuleNames.AddRange(
+			new string[]
+			{
+				"Core"
+			});
+
+		PrivateDependencyModuleNames.AddRange(
+			new string[]
+			{
+				"CoreUObject",
+				"Engine",
+				"Json",
+				"ModelContextProtocol",
+				"ModelContextProtocolEngine"
+			});
+	}
+}
diff --git a/Engine/Plugins/Experimental/ModelContextProtocol/Source/ModelContextProtocolEngineTests/Private/ModelContextProtocolEngineTestsModule.cpp b/Engine/Plugins/Experimental/ModelContextProtocol/Source/ModelContextProtocolEngineTests/Private/ModelContextProtocolEngineTestsModule.cpp
new file mode 100644
--- /dev/null
+++ b/Engine/Plugins/Experimental/ModelContextProtocol/Source/ModelContextProtocolEngineTests/Private/ModelContextProtocolEngineTestsModule.cpp
@@ -0,0 +1,9 @@
+// Copyright Epic Games, Inc. All Rights Reserved.
+
+#include "Modules/ModuleManager.h"
+
+class FModelContextProtocolEngineTestsModule : public IModuleInterface
+{
+};
+
+IMPLEMENT_MODULE(FModelContextProtocolEngineTestsModule, ModelContextProtocolEngineTests);
diff --git a/Engine/Plugins/Experimental/ModelContextProtocol/Source/ModelContextProtocolEngineTests/Private/Tests/ModelContextProtocolClientConfigTests.cpp b/Engine/Plugins/Experimental/ModelContextProtocol/Source/ModelContextProtocolEngineTests/Private/Tests/ModelContextProtocolClientConfigTests.cpp
new file mode 100644
--- /dev/null
+++ b/Engine/Plugins/Experimental/ModelContextProtocol/Source/ModelContextProtocolEngineTests/Private/Tests/ModelContextProtocolClientConfigTests.cpp
@@ -0,0 +1,176 @@
+// Copyright Epic Games, Inc. All Rights Reserved.
+
+#include "ModelContextProtocolClientConfig.h"
+#include "ModelContextProtocol.h"
+
+#include "Dom/JsonObject.h"
+#include "HAL/PlatformFileManager.h"
+#include "Misc/AutomationTest.h"
+#include "Misc/FileHelper.h"
+#include "Misc/Paths.h"
+#include "Serialization/JsonReader.h"
+#include "Serialization/JsonSerializer.h"
+
+#if WITH_DEV_AUTOMATION_TESTS
+
+BEGIN_DEFINE_SPEC(FModelContextProtocolClientConfigTests, "AI.ModelContextProtocol.ClientConfig", EAutomationTestFlags::ProductFilter | EAutomationTestFlags_ApplicationContextMask)
+	FString TestDirectory;
+
+	TSharedPtr<FJsonObject> ReadJsonFile(const FString& FilePath)
+	{
+		FString Content;
+		if (!FFileHelper::LoadFileToString(Content, *FilePath))
+		{
+			return nullptr;
+		}
+		TSharedPtr<FJsonObject> JsonObject;
+		TSharedRef<TJsonReader<>> Reader = TJsonReaderFactory<>::Create(Content);
+		FJsonSerializer::Deserialize(Reader, JsonObject);
+		return JsonObject;
+	}
+END_DEFINE_SPEC(FModelContextProtocolClientConfigTests)
+
+void FModelContextProtocolClientConfigTests::Define()
+{
+	BeforeEach([this]()
+	{
+		TestDirectory = FPaths::Combine(FPaths::ProjectSavedDir(), TEXT("Temp"), TEXT("MCPClientConfigTests"));
+		IFileManager::Get().MakeDirectory(*TestDirectory, true);
+	});
+
+	AfterEach([this]()
+	{
+		IFileManager::Get().DeleteDirectory(*TestDirectory, false, true);
+	});
+
+	Describe("WriteClientConfiguration", [this]()
+	{
+		It("should write ClaudeCode .mcp.json with correct format", [this]()
+		{
+			TestTrue("Should write successfully", UE::ModelContextProtocol::WriteClientConfiguration(EModelContextProtocolClient::ClaudeCode, 8000, TEXT("/mcp"), TestDirectory));
+
+			const FString FilePath = FPaths::Combine(TestDirectory, TEXT(".mcp.json"));
+			TSharedPtr<FJsonObject> Root = ReadJsonFile(FilePath);
+			if (!TestTrue("Should parse JSON", Root.IsValid())) { return; }
+
+			const TSharedPtr<FJsonObject>* Servers = nullptr;
+			if (!TestTrue("Should have mcpServers", Root->TryGetObjectField(TEXT("mcpServers"), Servers))) { return; }
+
+			const TSharedPtr<FJsonObject>* Entry = nullptr;
+			if (!TestTrue("Should have unreal-mcp entry", (*Servers)->TryGetObjectField(UE::ModelContextProtocol::DefaultServerName, Entry))) { return; }
+
+			FString Type, Url;
+			TestTrue("Should have type field", (*Entry)->TryGetStringField(TEXT("type"), Type));
+			TestEqual("Type should be http", Type, TEXT("http"));
+			TestTrue("Should have url field", (*Entry)->TryGetStringField(TEXT("url"), Url));
+			TestEqual("URL should match", Url, TEXT("http://127.0.0.1:8000/mcp"));
+		});
+
+		It("should write Cursor .cursor/mcp.json with correct format", [this]()
+		{
+			TestTrue("Should write successfully", UE::ModelContextProtocol::WriteClientConfiguration(EModelContextProtocolClient::Cursor, 9000, TEXT("/mcp"), TestDirectory));
+
+			const FString FilePath = FPaths::Combine(TestDirectory, TEXT(".cursor/mcp.json"));
+			TSharedPtr<FJsonObject> Root = ReadJsonFile(FilePath);
+			if (!TestTrue("Should parse JSON", Root.IsValid())) { return; }
+
+			const TSharedPtr<FJsonObject>* Servers = nullptr;
+			if (!TestTrue("Should have mcpServers", Root->TryGetObjectField(TEXT("mcpServers"), Servers))) { return; }
+
+			const TSharedPtr<FJsonObject>* Entry = nullptr;
+			if (!TestTrue("Should have unreal-mcp entry", (*Servers)->TryGetObjectField(UE::ModelContextProtocol::DefaultServerName, Entry))) { return; }
+
+			FString Url;
+			TestTrue("Should have url field", (*Entry)->TryGetStringField(TEXT("url"), Url));
+			TestEqual("URL should match", Url, TEXT("http://127.0.0.1:9000/mcp"));
+			TestFalse("Should NOT have type field", (*Entry)->HasField(TEXT("type")));
+		});
+
+		It("should write VSCode .vscode/mcp.json with servers root key", [this]()
+		{
+			TestTrue("Should write successfully", UE::ModelContextProtocol::WriteClientConfiguration(EModelContextProtocolClient::VSCode, 8000, TEXT("/mcp"), TestDirectory));
+
+			const FString FilePath = FPaths::Combine(TestDirectory, TEXT(".vscode/mcp.json"));
+			TSharedPtr<FJsonObject> Root = ReadJsonFile(FilePath);
+			if (!TestTrue("Should parse JSON", Root.IsValid())) { return; }
+
+			const TSharedPtr<FJsonObject>* Servers = nullptr;
+			TestTrue("Should have servers (not mcpServers)", Root->TryGetObjectField(TEXT("servers"), Servers));
+			TestFalse("Should NOT have mcpServers", Root->HasField(TEXT("mcpServers")));
+		});
+
+		It("should write Gemini .gemini/settings.json with httpUrl field", [this]()
+		{
+			TestTrue("Should write successfully", UE::ModelContextProtocol::WriteClientConfiguration(EModelContextProtocolClient::Gemini, 8000, TEXT("/mcp"), TestDirectory));
+
+			const FString FilePath = FPaths::Combine(TestDirectory, TEXT(".gemini/settings.json"));
+			TSharedPtr<FJsonObject> Root = ReadJsonFile(FilePath);
+			if (!TestTrue("Should parse JSON", Root.IsValid())) { return; }
+
+			const TSharedPtr<FJsonObject>* Servers = nullptr;
+			if (!TestTrue("Should have mcpServers", Root->TryGetObjectField(TEXT("mcpServers"), Servers))) { return; }
+
+			const TSharedPtr<FJsonObject>* Entry = nullptr;
+			if (!TestTrue("Should have unreal-mcp entry", (*Servers)->TryGetObjectField(UE::ModelContextProtocol::DefaultServerName, Entry))) { return; }
+
+			FString HttpUrl;
+			TestTrue("Should have httpUrl field", (*Entry)->TryGetStringField(TEXT("httpUrl"), HttpUrl));
+			TestEqual("httpUrl should match", HttpUrl, TEXT("http://127.0.0.1:8000/mcp"));
+			TestFalse("Should NOT have url field", (*Entry)->HasField(TEXT("url")));
+		});
+
+		It("should write Codex .codex/config.toml with correct format", [this]()
+		{
+			TestTrue("Should write successfully", UE::ModelContextProtocol::WriteClientConfiguration(EModelContextProtocolClient::Codex, 8000, TEXT("/mcp"), TestDirectory));
+
+			const FString FilePath = FPaths::Combine(TestDirectory, TEXT(".codex/config.toml"));
+			FString Content;
+			if (!TestTrue("Should read file", FFileHelper::LoadFileToString(Content, *FilePath))) { return; }
+
+			TestTrue("Should contain mcp_servers section", Content.Contains(TEXT("[mcp_servers.unreal-mcp]")));
+			TestTrue("Should contain url", Content.Contains(TEXT("url = \"http://127.0.0.1:8000/mcp\"")));
+		});
+
+		It("should refuse to overwrite existing Codex config", [this]()
+		{
+			// Write once
+			UE::ModelContextProtocol::WriteClientConfiguration(EModelContextProtocolClient::Codex, 8000, TEXT("/mcp"), TestDirectory);
+
+			// Attempt overwrite
+			AddExpectedError(TEXT("already exists"), EAutomationExpectedErrorFlags::Contains);
+			TestFalse("Should refuse to overwrite", UE::ModelContextProtocol::WriteClientConfiguration(EModelContextProtocolClient::Codex, 9000, TEXT("/mcp"), TestDirectory));
+		});
+
+		It("should preserve existing entries when updating JSON config", [this]()
+		{
+			const FString FilePath = FPaths::Combine(TestDirectory, TEXT(".mcp.json"));
+
+			// Write a pre-existing entry
+			TSharedPtr<FJsonObject> Root = MakeShared<FJsonObject>();
+			TSharedPtr<FJsonObject> Servers = MakeShared<FJsonObject>();
+			TSharedPtr<FJsonObject> ExistingEntry = MakeShared<FJsonObject>();
+			ExistingEntry->SetStringField(TEXT("command"), TEXT("my-server"));
+			Servers->SetObjectField(TEXT("my-existing-server"), ExistingEntry);
+			Root->SetObjectField(TEXT("mcpServers"), Servers);
+
+			FString OutputString;
+			TSharedRef<TJsonWriter<>> Writer = TJsonWriterFactory<>::Create(&OutputString);
+			FJsonSerializer::Serialize(Root.ToSharedRef(), Writer);
+			FFileHelper::SaveStringToFile(OutputString, *FilePath);
+
+			// Now write the MCP config
+			TestTrue("Should write successfully", UE::ModelContextProtocol::WriteClientConfiguration(EModelContextProtocolClient::ClaudeCode, 8000, TEXT("/mcp"), TestDirectory));
+
+			TSharedPtr<FJsonObject> UpdatedRoot = ReadJsonFile(FilePath);
+			if (!TestTrue("Should parse updated JSON", UpdatedRoot.IsValid())) { return; }
+
+			const TSharedPtr<FJsonObject>* UpdatedServers = nullptr;
+			if (!TestTrue("Should have mcpServers", UpdatedRoot->TryGetObjectField(TEXT("mcpServers"), UpdatedServers))) { return; }
+
+			TestTrue("Should still have existing server", (*UpdatedServers)->HasField(TEXT("my-existing-server")));
+			TestTrue("Should also have unreal-mcp", (*UpdatedServers)->HasField(UE::ModelContextProtocol::DefaultServerName));
+		});
+	});
+}
+
+#endif // WITH_DEV_AUTOMATION_TESTS
