[ModelContextProtocol] Provide default `inputSchema` in `IModelContextProtocolTool` base class.

Per the MCP specification, `inputSchema` is required on every tool entry and must be a *JSON* Schema object with `"type": "object"`. The base class `IModelContextProtocolTool::GetInputJsonSchema()` was returning null by default, causing tools that did not override it to break the entire `tools/list` response with client-side validation errors. The default now returns `{"type": "object"}` for convenience.

#Synced-CL 52658100

diff --git a/Engine/Plugins/Experimental/ModelContextProtocol/Source/ModelContextProtocol/Public/IModelContextProtocolTool.h b/Engine/Plugins/Experimental/ModelContextProtocol/Source/ModelContextProtocol/Public/IModelContextProtocolTool.h
--- a/Engine/Plugins/Experimental/ModelContextProtocol/Source/ModelContextProtocol/Public/IModelContextProtocolTool.h
+++ b/Engine/Plugins/Experimental/ModelContextProtocol/Source/ModelContextProtocol/Public/IModelContextProtocolTool.h
@@ -38,12 +38,19 @@ struct IModelContextProtocolTool : TSharedFromThis<IModelContextProtocolTool>
 	virtual FString GetDescription() const = 0;
 
 	/**
-	 * Returns an optional JSON Schema definition which the input parameters to this tool must adhere to.
-	 * @return A valid JSON Schema definition or invalid shared pointer for undefined inputs.  
+	 * Returns the JSON Schema definition which the input parameters to this tool must adhere to.
+	 * Per the MCP specification, inputSchema is required and must be a JSON Schema object
+	 * with "type": "object". The default returns {"type": "object"} (accepts any parameters).
+	 * @return A valid JSON Schema definition.
 	 * @see https://json-schema.org
 	 * @see FJsonSchemaGenerator::UStructToJsonSchemaObject
 	 */
-	virtual TSharedPtr<FJsonObject> GetInputJsonSchema() const { return {}; }
+	virtual TSharedPtr<FJsonObject> GetInputJsonSchema() const
+	{
+		FJsonDomBuilder::FObject Schema;
+		Schema.Set(TEXT("type"), TEXT("object"));
+		return Schema.AsJsonObject().ToSharedPtr();
+	}
 
 	/**
 	 * Returns an optional JSON Schema definition which the results of this tool must adhere to.
diff --git a/Engine/Plugins/Experimental/ModelContextProtocol/Source/ModelContextProtocolTests/Private/Tests/ModelContextProtocolModuleTests.cpp b/Engine/Plugins/Experimental/ModelContextProtocol/Source/ModelContextProtocolTests/Private/Tests/ModelContextProtocolModuleTests.cpp
--- a/Engine/Plugins/Experimental/ModelContextProtocol/Source/ModelContextProtocolTests/Private/Tests/ModelContextProtocolModuleTests.cpp
+++ b/Engine/Plugins/Experimental/ModelContextProtocol/Source/ModelContextProtocolTests/Private/Tests/ModelContextProtocolModuleTests.cpp
@@ -45,6 +45,30 @@ void FModelContextProtocolModuleTests::Define()
 		MockProvider.Reset();
 	});
 
+	Describe("Default inputSchema", [this]()
+	{
+		It("should return a valid JSON Schema object from the base class GetInputJsonSchema", [this]()
+		{
+			// A minimal tool that does not override GetInputJsonSchema — it should inherit the base class default.
+			struct FMinimalTool : IModelContextProtocolTool
+			{
+				virtual FString GetName() const override { return TEXT("test_minimal"); }
+				virtual FString GetDescription() const override { return TEXT(""); }
+			};
+
+			FMinimalTool Tool;
+			TSharedPtr<FJsonObject> Schema = Tool.GetInputJsonSchema();
+			if (TestTrue("Default GetInputJsonSchema should return a valid object", Schema.IsValid()))
+			{
+				FString TypeValue;
+				if (TestTrue("Schema should have a 'type' field", Schema->TryGetStringField(TEXT("type"), TypeValue)))
+				{
+					TestEqual("Schema type should be 'object'", TypeValue, TEXT("object"));
+				}
+			}
+		});
+	});
+
 	Describe("Tool management", [this]()
 	{
 		It("should register a tool via AddTool and retrieve it via FindTool", [this]()
