[MCP] Add deferred tool loading to ModelContextProtocol server

Reduce context window usage from ~120k tokens to ~550 by deferring tool schema loading. When enabled (on by default), tools/list returns three lightweight discovery tools instead of all toolset schemas. The LLM discovers toolsets via list_toolsets, inspects them via describe_toolset, and loads them on demand via load_toolset, which registers native MCP tools and broadcasts notifications/tools/list_changed for clients to re-fetch.

Changes:
- Add ModelContextProtocol.DeferredToolLoading CVar (default: true)
- Add list_toolsets, describe_toolset, load_toolset MCP tools for progressive discovery
- Split ToolsetRegistryToolAdapterManager into eager/deferred registration modes
- Make FToolset::GetToolsetDescription() pure virtual, implement in all subclasses
- Fix TWeakObjectPtr Pin() pattern in BlueprintLibraryToolset::GetToolsetDescription
- load_toolset returns lightweight summary instead of full schemas (tools are registered as native MCP tools)
- Add comprehensive tests covering both eager and deferred modes (14 tests)

#Synced-CL 52732140

#Synced-CL 52732140

diff --git a/Engine/Plugins/Experimental/AIAssistant/Source/AIAssistant/Private/Tests/AIAssistantFakeToolset.cpp b/Engine/Plugins/Experimental/AIAssistant/Source/AIAssistant/Private/Tests/AIAssistantFakeToolset.cpp
--- a/Engine/Plugins/Experimental/AIAssistant/Source/AIAssistant/Private/Tests/AIAssistantFakeToolset.cpp
+++ b/Engine/Plugins/Experimental/AIAssistant/Source/AIAssistant/Private/Tests/AIAssistantFakeToolset.cpp
@@ -62,6 +62,11 @@ namespace UE::AIAssistant
 		return TEXT("1.0");
 	}
 
+	FString FFakeToolset::GetToolsetDescription() const
+	{
+		return TEXT("Fake Toolset for testing.");
+	}
+
 	TSharedPtr<TPromise<TValueOrError<FString, FString>>> FFakeToolset::AddFakeToolCall(
 		const FString& ToolName,
 		TPromise<TValueOrError<FString, FString>>&& Promise)
diff --git a/Engine/Plugins/Experimental/AIAssistant/Source/AIAssistant/Private/Tests/AIAssistantFakeToolset.h b/Engine/Plugins/Experimental/AIAssistant/Source/AIAssistant/Private/Tests/AIAssistantFakeToolset.h
--- a/Engine/Plugins/Experimental/AIAssistant/Source/AIAssistant/Private/Tests/AIAssistantFakeToolset.h
+++ b/Engine/Plugins/Experimental/AIAssistant/Source/AIAssistant/Private/Tests/AIAssistantFakeToolset.h
@@ -33,6 +33,7 @@ public:
 	FString GetToolsetName() const override;
 
 	FString GetToolsetVersion() const override;
+	FString GetToolsetDescription() const override;
 
 	TSharedPtr<TPromise<TValueOrError<FString, FString>>> AddFakeToolCall(
 		const FString& ToolName,
diff --git a/Engine/Plugins/Experimental/ModelContextProtocol/Source/ModelContextProtocol/Private/ModelContextProtocol.cpp b/Engine/Plugins/Experimental/ModelContextProtocol/Source/ModelContextProtocol/Private/ModelContextProtocol.cpp
--- a/Engine/Plugins/Experimental/ModelContextProtocol/Source/ModelContextProtocol/Private/ModelContextProtocol.cpp
+++ b/Engine/Plugins/Experimental/ModelContextProtocol/Source/ModelContextProtocol/Private/ModelContextProtocol.cpp
@@ -65,4 +65,11 @@ namespace UE::ModelContextProtocol
 		TEXT("ModelContextProtocol.PaginationPageSize"),
 		PaginationPageSize,
 		TEXT("Maximum number of items returned per page for list methods (tools/list, resources/list). 0 = disabled (return all items)."));
+
+	bool bDeferredToolLoading = true;
+	static FAutoConsoleVariableRef CVarDeferredToolLoading(
+		TEXT("ModelContextProtocol.DeferredToolLoading"),
+		bDeferredToolLoading,
+		TEXT("When enabled, tools/list returns list_toolsets, describe_toolset, and load_toolset instead of all tool schemas. ")
+		TEXT("The LLM discovers and loads specific toolsets on demand, reducing context window usage."));
 }
diff --git a/Engine/Plugins/Experimental/ModelContextProtocol/Source/ModelContextProtocol/Public/ModelContextProtocol.h b/Engine/Plugins/Experimental/ModelContextProtocol/Source/ModelContextProtocol/Public/ModelContextProtocol.h
--- a/Engine/Plugins/Experimental/ModelContextProtocol/Source/ModelContextProtocol/Public/ModelContextProtocol.h
+++ b/Engine/Plugins/Experimental/ModelContextProtocol/Source/ModelContextProtocol/Public/ModelContextProtocol.h
@@ -71,6 +71,7 @@ namespace UE::ModelContextProtocol
 	UE_API extern bool bAudioResultOggFormat;
 	UE_API extern float ProgressIntervalSeconds;
 	UE_API extern int32 PaginationPageSize;
+	UE_API extern bool bDeferredToolLoading;
 }
 
 #undef UE_API
diff --git a/Engine/Plugins/Experimental/ModelContextProtocol/Source/ModelContextProtocolEditor/Private/ModelContextProtocolDeferredToolLoader.cpp b/Engine/Plugins/Experimental/ModelContextProtocol/Source/ModelContextProtocolEditor/Private/ModelContextProtocolDeferredToolLoader.cpp
new file mode 100644
--- /dev/null
+++ b/Engine/Plugins/Experimental/ModelContextProtocol/Source/ModelContextProtocolEditor/Private/ModelContextProtocolDeferredToolLoader.cpp
@@ -0,0 +1,119 @@
+// Copyright Epic Games, Inc. All Rights Reserved.
+
+#include "ModelContextProtocolDeferredToolLoader.h"
+
+#include "Dom/JsonObject.h"
+#include "ModelContextProtocol.h"
+#include "ModelContextProtocolToolResults.h"
+
+namespace
+{
+	TSharedPtr<FJsonObject> MakeToolsetNameSchema(const FString& Description)
+	{
+		TSharedPtr<FJsonObject> Schema = MakeShared<FJsonObject>();
+		Schema->SetStringField(TEXT("type"), TEXT("object"));
+
+		TSharedPtr<FJsonObject> Properties = MakeShared<FJsonObject>();
+		TSharedPtr<FJsonObject> ToolsetNameProp = MakeShared<FJsonObject>();
+		ToolsetNameProp->SetStringField(TEXT("type"), TEXT("string"));
+		ToolsetNameProp->SetStringField(TEXT("description"), Description);
+		Properties->SetObjectField(TEXT("toolset_name"), ToolsetNameProp);
+		Schema->SetObjectField(TEXT("properties"), Properties);
+
+		TArray<TSharedPtr<FJsonValue>> RequiredArray;
+		RequiredArray.Add(MakeShared<FJsonValueString>(TEXT("toolset_name")));
+		Schema->SetArrayField(TEXT("required"), RequiredArray);
+
+		return Schema;
+	}
+}
+
+// -- FListToolsetsTool --
+
+FListToolsetsTool::FListToolsetsTool(FListToolsetsDelegate InDelegate)
+	: ListDelegate(MoveTemp(InDelegate))
+{
+}
+
+TSharedPtr<FJsonObject> FListToolsetsTool::GetInputJsonSchema() const
+{
+	TSharedPtr<FJsonObject> Schema = MakeShared<FJsonObject>();
+	Schema->SetStringField(TEXT("type"), TEXT("object"));
+	Schema->SetObjectField(TEXT("properties"), MakeShared<FJsonObject>());
+	return Schema;
+}
+
+FModelContextProtocolToolResult FListToolsetsTool::Run(const TSharedPtr<FJsonObject>& Params)
+{
+	FString Result = ListDelegate();
+	return UE::ModelContextProtocol::MakeTextResult(Result);
+}
+
+// -- FDescribeToolsetTool --
+
+FDescribeToolsetTool::FDescribeToolsetTool(FDescribeToolsetDelegate InDelegate)
+	: DescribeDelegate(MoveTemp(InDelegate))
+{
+}
+
+TSharedPtr<FJsonObject> FDescribeToolsetTool::GetInputJsonSchema() const
+{
+	return MakeToolsetNameSchema(TEXT("Name of the toolset to describe. Use list_toolsets to see available names."));
+}
+
+FModelContextProtocolToolResult FDescribeToolsetTool::Run(const TSharedPtr<FJsonObject>& Params)
+{
+	if (!Params.IsValid())
+	{
+		return UE::ModelContextProtocol::MakeErrorResult(TEXT("Missing parameters."));
+	}
+
+	FString ToolsetName;
+	if (!Params->TryGetStringField(TEXT("toolset_name"), ToolsetName) || ToolsetName.IsEmpty())
+	{
+		return UE::ModelContextProtocol::MakeErrorResult(TEXT("Missing required parameter: toolset_name"));
+	}
+
+	TValueOrError<FString, FString> Result = DescribeDelegate(ToolsetName);
+	if (Result.HasError())
+	{
+		return UE::ModelContextProtocol::MakeErrorResult(Result.GetError());
+	}
+
+	return UE::ModelContextProtocol::MakeTextResult(Result.GetValue());
+}
+
+// -- FLoadToolsetTool --
+
+FLoadToolsetTool::FLoadToolsetTool(FLoadToolsetDelegate InDelegate)
+	: LoadDelegate(MoveTemp(InDelegate))
+{
+}
+
+TSharedPtr<FJsonObject> FLoadToolsetTool::GetInputJsonSchema() const
+{
+	return MakeToolsetNameSchema(TEXT("Name of the toolset to load. Use list_toolsets to see available names."));
+}
+
+FModelContextProtocolToolResult FLoadToolsetTool::Run(const TSharedPtr<FJsonObject>& Params)
+{
+	if (!Params.IsValid())
+	{
+		return UE::ModelContextProtocol::MakeErrorResult(TEXT("Missing parameters."));
+	}
+
+	FString ToolsetName;
+	if (!Params->TryGetStringField(TEXT("toolset_name"), ToolsetName) || ToolsetName.IsEmpty())
+	{
+		return UE::ModelContextProtocol::MakeErrorResult(TEXT("Missing required parameter: toolset_name"));
+	}
+
+	FString ResultText;
+	bool bSuccess = LoadDelegate(ToolsetName, ResultText);
+	if (!bSuccess)
+	{
+		return UE::ModelContextProtocol::MakeErrorResult(ResultText);
+	}
+
+	return UE::ModelContextProtocol::MakeTextResult(ResultText);
+}
diff --git a/Engine/Plugins/Experimental/ModelContextProtocol/Source/ModelContextProtocolEditor/Private/ModelContextProtocolDeferredToolLoader.h b/Engine/Plugins/Experimental/ModelContextProtocol/Source/ModelContextProtocolEditor/Private/ModelContextProtocolDeferredToolLoader.h
new file mode 100644
--- /dev/null
+++ b/Engine/Plugins/Experimental/ModelContextProtocol/Source/ModelContextProtocolEditor/Private/ModelContextProtocolDeferredToolLoader.h
@@ -0,0 +1,80 @@
+// Copyright Epic Games, Inc. All Rights Reserved.
+
+#pragma once
+
+#include "IModelContextProtocolTool.h"
+#include "Templates/ValueOrError.h"
+
+/** Callback: returns a text listing of all available toolsets (name + description). */
+using FListToolsetsDelegate = TFunction<FString()>;
+
+/** Callback: returns the full JSON schema for a named toolset, or an error string. */
+using FDescribeToolsetDelegate = TFunction<TValueOrError<FString, FString>(const FString& ToolsetName)>;
+
+/** Callback: loads a toolset by name, registering its tools as native MCP tools. Returns true on success. */
+using FLoadToolsetDelegate = TFunction<bool(const FString& ToolsetName, FString& OutResultText)>;
+
+/**
+ * Lists all available toolsets with their names and descriptions.
+ * Use this to discover what toolsets are available before loading them.
+ */
+struct FListToolsetsTool : IModelContextProtocolTool
+{
+	explicit FListToolsetsTool(FListToolsetsDelegate InDelegate);
+
+	virtual FString GetName() const override { return TEXT("list_toolsets"); }
+	virtual FString GetDescription() const override { return TEXT("List all available toolsets with names and descriptions."); }
+	virtual TSharedPtr<FJsonObject> GetInputJsonSchema() const override;
+	virtual FModelContextProtocolToolResult Run(const TSharedPtr<FJsonObject>& Params) override;
+
+private:
+	FListToolsetsDelegate ListDelegate;
+};
+
+/**
+ * Returns detailed information about a specific toolset including all tool names,
+ * descriptions, and input schemas. Use this to inspect a toolset before loading it.
+ */
+struct FDescribeToolsetTool : IModelContextProtocolTool
+{
+	explicit FDescribeToolsetTool(FDescribeToolsetDelegate InDelegate);
+
+	virtual FString GetName() const override { return TEXT("describe_toolset"); }
+	virtual FString GetDescription() const override
+	{
+		return TEXT("Get detailed information about a toolset including all "
+			"tool names, descriptions, and input schemas.");
+	}
+	virtual TSharedPtr<FJsonObject> GetInputJsonSchema() const override;
+	virtual FModelContextProtocolToolResult Run(const TSharedPtr<FJsonObject>& Params) override;
+
+private:
+	FDescribeToolsetDelegate DescribeDelegate;
+};
+
+/**
+ * Loads a toolset by name, registering its individual tools as native MCP tools.
+ * After loading, the tools become available starting from the next client turn
+ * (not in the same response as the load call).
+ * Use list_toolsets to discover available toolsets.
+ */
+struct FLoadToolsetTool : IModelContextProtocolTool
+{
+	explicit FLoadToolsetTool(FLoadToolsetDelegate InDelegate);
+
+	virtual FString GetName() const override { return TEXT("load_toolset"); }
+	virtual FString GetDescription() const override
+	{
+		return TEXT("Load a toolset to register its tools as native MCP tools. "
+			"Use list_toolsets to discover available toolsets. "
+			"IMPORTANT: loaded tools are available starting from the next turn, "
+			"not in the same response as this call. Do not attempt to use them "
+			"in the same turn. The tool names follow the client's MCP naming "
+			"convention, not the names shown in this result.");
+	}
+	virtual TSharedPtr<FJsonObject> GetInputJsonSchema() const override;
+	virtual FModelContextProtocolToolResult Run(const TSharedPtr<FJsonObject>& Params) override;
+
+private:
+	FLoadToolsetDelegate LoadDelegate;
+};
diff --git a/Engine/Plugins/Experimental/ModelContextProtocol/Source/ModelContextProtocolEditor/Private/ModelContextProtocolToolsetRegistryAdapter.cpp b/Engine/Plugins/Experimental/ModelContextProtocol/Source/ModelContextProtocolEditor/Private/ModelContextProtocolToolsetRegistryAdapter.cpp
--- a/Engine/Plugins/Experimental/ModelContextProtocol/Source/ModelContextProtocolEditor/Private/ModelContextProtocolToolsetRegistryAdapter.cpp
+++ b/Engine/Plugins/Experimental/ModelContextProtocol/Source/ModelContextProtocolEditor/Private/ModelContextProtocolToolsetRegistryAdapter.cpp
@@ -1,16 +1,30 @@
 // Copyright Epic Games, Inc. All Rights Reserved.
 
 #include "ModelContextProtocolToolsetRegistryAdapter.h"
+
 #include "IModelContextProtocolModule.h"
 #include "ModelContextProtocol.h"
+#include "ModelContextProtocolDeferredToolLoader.h"
+#include "ModelContextProtocolServer.h"
 #include "ModelContextProtocolToolResults.h"
 
 #include "Editor.h"
 #include "Serialization/JsonReader.h"
 #include "Serialization/JsonSerializer.h"
+#include "ToolsetRegistry/Toolset.h"
 #include "ToolsetRegistry/ToolsetRegistry.h"
 #include "ToolsetRegistry/ToolsetRegistrySubsystem.h"
 
+namespace
+{
+	UE::ToolsetRegistry::FToolsetRegistry* GetToolsetRegistry()
+	{
+		UToolsetRegistrySubsystem* Subsystem = GEditor ?
+			GEditor->GetEditorSubsystem<UToolsetRegistrySubsystem>() : nullptr;
+		return Subsystem ? &Subsystem->ToolsetRegistry : nullptr;
+	}
+}
+
 FToolsetRegistryToolAdapter::FToolsetRegistryToolAdapter(const FString& InToolName, const FString& InDescription, TSharedPtr<FJsonObject> InInputSchema)
 	: ToolName(InToolName)
 	, Description(InDescription)
@@ -20,8 +34,8 @@ FToolsetRegistryToolAdapter::FToolsetRegistryToolAdapter(const FString& InToolNa
 
 void FToolsetRegistryToolAdapter::RunAsync(const FModelContextProtocolToolRequestId& RequestId, const TSharedPtr<FJsonObject>& Params, const FResultCallback& OnComplete)
 {
-	UToolsetRegistrySubsystem* ToolsetRegistrySubsystem = GEditor ? GEditor->GetEditorSubsystem<UToolsetRegistrySubsystem>() : nullptr;
-	if (!ToolsetRegistrySubsystem)
+	UE::ToolsetRegistry::FToolsetRegistry* ToolsetRegistry = GetToolsetRegistry();
+	if (!ToolsetRegistry)
 	{
 		OnComplete(UE::ModelContextProtocol::MakeErrorResult(TEXT("ToolsetRegistry not available")));
 		return;
@@ -47,8 +61,7 @@ void FToolsetRegistryToolAdapter::RunAsync(const FModelContextProtocolToolReques
 
 	UE_LOGF(LogModelContextProtocol, Log, "Running ToolsetRegistry tool: '%ls'", *ToolName);
 
-	UE::ToolsetRegistry::FToolsetRegistry& ToolsetRegistry = ToolsetRegistrySubsystem->ToolsetRegistry;
-	ToolsetRegistry.ExecuteTool(DescriptorResult.GetValue(), ArgumentsJson).Then([OnComplete](TFuture<TValueOrError<FString, FString>> Future)
+	ToolsetRegistry->ExecuteTool(DescriptorResult.GetValue(), ArgumentsJson).Then([OnComplete](TFuture<TValueOrError<FString, FString>> Future)
 	{
 		TValueOrError<FString, FString> Result = Future.Get();
 
@@ -63,12 +76,71 @@ void FToolsetRegistryToolAdapter::RunAsync(const FModelContextProtocolToolReques
 	});
 }
 
+int32 FToolsetRegistryToolAdapterManager::RegisterToolsFromSchema(const FString& ToolsetSchemaJson)
+{
+	IModelContextProtocolModule* const Module = IModelContextProtocolModule::Get();
+	if (!Module)
+	{
+		return 0;
+	}
+
+	TSharedPtr<FJsonObject> ToolsetSchema;
+	TSharedRef<TJsonReader<>> Reader = TJsonReaderFactory<>::Create(ToolsetSchemaJson);
+	if (!FJsonSerializer::Deserialize(Reader, ToolsetSchema) || !ToolsetSchema.IsValid())
+	{
+		return 0;
+	}
+
+	const TArray<TSharedPtr<FJsonValue>>* ToolsField = nullptr;
+	if (!ToolsetSchema->TryGetArrayField(TEXT("tools"), ToolsField) || !ToolsField)
+	{
+		return 0;
+	}
+
+	int32 RegisteredCount = 0;
+	for (const TSharedPtr<FJsonValue>& ToolValue : *ToolsField)
+	{
+		const TSharedPtr<FJsonObject>& ToolJsonObject = ToolValue->AsObject();
+		if (!ToolJsonObject.IsValid())
+		{
+			continue;
+		}
+
+		FString ToolName;
+		ToolJsonObject->TryGetStringField(TEXT("name"), ToolName);
+		if (ToolName.IsEmpty())
+		{
+			continue;
+		}
+
+		FString ToolDescription;
+		ToolJsonObject->TryGetStringField(TEXT("description"), ToolDescription);
+
+		TSharedPtr<FJsonObject> InputSchema;
+		const TSharedPtr<FJsonObject>* InputSchemaField = nullptr;
+		if (ToolJsonObject->TryGetObjectField(TEXT("inputSchema"), InputSchemaField) && InputSchemaField)
+		{
+			InputSchema = *InputSchemaField;
+		}
+
+		TSharedRef<FToolsetRegistryToolAdapter> Adapter =
+			MakeShared<FToolsetRegistryToolAdapter>(ToolName, ToolDescription, InputSchema);
+		if (Module->AddTool(Adapter))
+		{
+			RegisteredAdapters.Add(Adapter);
+			RegisteredCount++;
+		}
+	}
+
+	return RegisteredCount;
+}
+
 void FToolsetRegistryToolAdapterManager::RegisterTools()
 {
-	DeregisterTools();
+	DeregisterTools(false);
 
-	UToolsetRegistrySubsystem* ToolsetRegistrySubsystem = GEditor ? GEditor->GetEditorSubsystem<UToolsetRegistrySubsystem>() : nullptr;
-	if (!ToolsetRegistrySubsystem)
+	UE::ToolsetRegistry::FToolsetRegistry* ToolsetRegistry = GetToolsetRegistry();
+	if (!ToolsetRegistry)
 	{
 		return;
 	}
@@ -79,79 +151,240 @@ void FToolsetRegistryToolAdapterManager::RegisterTools()
 		return;
 	}
 
-	UE::ToolsetRegistry::FToolsetRegistry& ToolsetRegistry = ToolsetRegistrySubsystem->ToolsetRegistry;
-	const FString SchemasJson = ToolsetRegistry.GetToolsetJsonSchemas();
+	if (UE::ModelContextProtocol::bDeferredToolLoading)
+	{
+		// Deferred mode: register discovery and loading tools
+		// instead of all tool schemas.
+		auto RegisterDeferredTool =
+			[Module, this](const TSharedRef<IModelContextProtocolTool>& Tool)
+		{
+			if (Module->AddTool(Tool))
+			{
+				DeferredTools.Add(Tool);
+			}
+		};
+
+		RegisterDeferredTool(MakeShared<FListToolsetsTool>(
+			FListToolsetsDelegate([this]()
+			{
+				return GetToolsetCatalogText();
+			})));
+
+		RegisterDeferredTool(MakeShared<FDescribeToolsetTool>(
+			FDescribeToolsetDelegate([this](const FString& ToolsetName)
+			{
+				return GetToolsetSchemaText(ToolsetName);
+			})));
 
-	TArray<TSharedPtr<FJsonValue>> ToolsetArray;
-	TSharedRef<TJsonReader<>> Reader = TJsonReaderFactory<>::Create(SchemasJson);
-	if (!FJsonSerializer::Deserialize(Reader, ToolsetArray))
+		RegisterDeferredTool(MakeShared<FLoadToolsetTool>(
+			FLoadToolsetDelegate(
+				[this](const FString& ToolsetName, FString& OutResultText)
+			{
+				return RegisterToolset(ToolsetName, OutResultText);
+			})));
+
+		int32 ToolsetCount = 0;
+		ToolsetRegistry->ForEachToolset(
+			[&ToolsetCount](const FString&, const UE::ToolsetRegistry::FToolset&)
+		{
+			ToolsetCount++;
+		});
+
+		UE_LOGF(LogModelContextProtocol, Log,
+			"Deferred tool loading: registered %d deferred tools "
+			"with %d toolsets in catalog",
+			DeferredTools.Num(), ToolsetCount);
+	}
+	else
 	{
-		return;
+		// Eager mode: register all tools from all toolsets.
+		ToolsetRegistry->ForEachToolset(
+			[this](const FString&, const UE::ToolsetRegistry::FToolset& Toolset)
+		{
+			RegisterToolsFromSchema(Toolset.GetToolsetJsonSchema());
+		});
+
+		UE_LOGF(LogModelContextProtocol, Log,
+			"Registered %d ToolsetRegistry tool adapters",
+			RegisteredAdapters.Num());
 	}
+}
 
-	for (const TSharedPtr<FJsonValue>& ToolsetValue : ToolsetArray)
+void FToolsetRegistryToolAdapterManager::DeregisterTools(bool bBroadcast)
+{
+	if (IModelContextProtocolModule* Module = IModelContextProtocolModule::Get())
 	{
-		const TSharedPtr<FJsonObject>& ToolsetObject = ToolsetValue->AsObject();
-		if (!ToolsetObject.IsValid())
+		for (const TSharedRef<FToolsetRegistryToolAdapter>& Adapter : RegisteredAdapters)
 		{
-			continue;
+			Module->RemoveTool(Adapter);
 		}
 
-		const TArray<TSharedPtr<FJsonValue>>* ToolsField = nullptr;
-		if (!ToolsetObject->TryGetArrayField(TEXT("tools"), ToolsField) || !ToolsField)
+		for (const TSharedRef<IModelContextProtocolTool>& Tool : DeferredTools)
 		{
-			continue;
+			Module->RemoveTool(Tool);
 		}
+	}
 
-		for (const TSharedPtr<FJsonValue>& ToolValue : *ToolsField)
+	RegisteredAdapters.Reset();
+	DeferredTools.Reset();
+	LoadedToolsetNames.Reset();
+
+	if (bBroadcast)
+	{
+		if (IModelContextProtocolModule* const Module = IModelContextProtocolModule::Get())
 		{
-			const TSharedPtr<FJsonObject>& ToolJsonObject = ToolValue->AsObject();
-			if (!ToolJsonObject.IsValid())
+			if (FModelContextProtocolServer* const Server = Module->GetServer())
 			{
-				continue;
+				Server->BroadcastToolsListChanged();
 			}
+		}
+	}
+}
 
-			FString ToolName;
-			ToolJsonObject->TryGetStringField(TEXT("name"), ToolName);
-			if (ToolName.IsEmpty())
-			{
-				continue;
-			}
+bool FToolsetRegistryToolAdapterManager::RegisterToolset(
+	const FString& ToolsetName, FString& OutResultText)
+{
+	if (LoadedToolsetNames.Contains(ToolsetName))
+	{
+		OutResultText = FString::Printf(
+			TEXT("Toolset '%s' is already loaded."), *ToolsetName);
+		return true;
+	}
 
-			FString ToolDescription;
-			ToolJsonObject->TryGetStringField(TEXT("description"), ToolDescription);
+	UE::ToolsetRegistry::FToolsetRegistry* ToolsetRegistry = GetToolsetRegistry();
+	if (!ToolsetRegistry)
+	{
+		OutResultText = TEXT("ToolsetRegistry not available.");
+		return false;
+	}
 
-			TSharedPtr<FJsonObject> InputSchema;
-			const TSharedPtr<FJsonObject>* InputSchemaField = nullptr;
-			if (ToolJsonObject->TryGetObjectField(TEXT("inputSchema"), InputSchemaField) && InputSchemaField)
-			{
-				InputSchema = *InputSchemaField;
-			}
+	TSharedPtr<UE::ToolsetRegistry::FToolset> Toolset = ToolsetRegistry->Find(ToolsetName);
+	if (!Toolset.IsValid())
+	{
+		TArray<FString> Names;
+		ToolsetRegistry->ForEachToolset(
+			[&Names](const FString& Name, const UE::ToolsetRegistry::FToolset&)
+		{
+			Names.Add(Name);
+		});
+		OutResultText = FString::Printf(
+			TEXT("Toolset '%s' not found. Available toolsets: %s"),
+			*ToolsetName, *FString::Join(Names, TEXT(", ")));
+		return false;
+	}
 
-			TSharedRef<FToolsetRegistryToolAdapter> Adapter = MakeShared<FToolsetRegistryToolAdapter>(ToolName, ToolDescription, InputSchema);
-			if (Module->AddTool(Adapter))
+	const FString SchemaJson = Toolset->GetToolsetJsonSchema();
+	int32 ToolCount = RegisterToolsFromSchema(SchemaJson);
+	if (ToolCount == 0)
+	{
+		OutResultText = FString::Printf(
+			TEXT("Toolset '%s' found but contains no tools."),
+			*ToolsetName);
+		return false;
+	}
+
+	LoadedToolsetNames.Add(ToolsetName);
+
+	// Broadcast tools/list_changed so clients re-fetch the tool list.
+	if (IModelContextProtocolModule* const Module = IModelContextProtocolModule::Get())
+	{
+		if (FModelContextProtocolServer* const Server = Module->GetServer())
+		{
+			Server->BroadcastToolsListChanged();
+		}
+	}
+
+	// Build a lightweight summary. Full schemas are redundant here
+	// because the tools are now registered as native MCP tools and
+	// the client will re-fetch them via tools/list_changed.
+	OutResultText = FString::Printf(
+		TEXT("Loaded %d tools from '%s':\n"), ToolCount, *ToolsetName);
+
+	// Reuse SchemaJson from above to list tool names in the summary.
+	TSharedPtr<FJsonObject> Schema;
+	TSharedRef<TJsonReader<>> Reader = TJsonReaderFactory<>::Create(SchemaJson);
+	if (FJsonSerializer::Deserialize(Reader, Schema) && Schema.IsValid())
+	{
+		const TArray<TSharedPtr<FJsonValue>>* ToolsField = nullptr;
+		if (Schema->TryGetArrayField(TEXT("tools"), ToolsField) && ToolsField)
+		{
+			for (const TSharedPtr<FJsonValue>& ToolValue : *ToolsField)
 			{
-				RegisteredAdapters.Add(Adapter);
+				const TSharedPtr<FJsonObject>& ToolObj = ToolValue->AsObject();
+				if (!ToolObj.IsValid())
+				{
+					continue;
+				}
+
+				FString ToolName;
+				ToolObj->TryGetStringField(TEXT("name"), ToolName);
+
+				FString ToolDescription;
+				ToolObj->TryGetStringField(TEXT("description"), ToolDescription);
+				if (ToolDescription.Len() > 80)
+				{
+					ToolDescription.LeftInline(77);
+					ToolDescription += TEXT("...");
+				}
+
+				OutResultText += FString::Printf(
+					TEXT("- %s -- %s\n"), *ToolName, *ToolDescription);
 			}
 		}
 	}
 
-	UE_LOGF(LogModelContextProtocol, Log, "Registered %d ToolsetRegistry tool adapters", RegisteredAdapters.Num());
+	UE_LOGF(
+		LogModelContextProtocol, Log, "Deferred tool loading: loaded toolset '%ls' (%d tools)",
+		*ToolsetName, ToolCount);
+
+	return true;
 }
 
-void FToolsetRegistryToolAdapterManager::DeregisterTools()
+FString FToolsetRegistryToolAdapterManager::GetToolsetCatalogText() const
 {
-	if (RegisteredAdapters.IsEmpty())
+	UE::ToolsetRegistry::FToolsetRegistry* ToolsetRegistry = GetToolsetRegistry();
+	if (!ToolsetRegistry)
 	{
-		return;
+		return FString();
 	}
 
-	if (IModelContextProtocolModule* Module = IModelContextProtocolModule::Get())
-	{
-		for (const TSharedRef<FToolsetRegistryToolAdapter>& Adapter : RegisteredAdapters)
+	FString Result;
+	ToolsetRegistry->ForEachToolset(
+		[&Result](const FString& Name, const UE::ToolsetRegistry::FToolset& Toolset)
 		{
-			Module->RemoveTool(Adapter);
-		}
+			Result += FString::Printf(TEXT("- %s"), *Name);
+			FString Description = Toolset.GetToolsetDescription();
+			if (!Description.IsEmpty())
+			{
+				Result += FString::Printf(TEXT(": %s"), *Description);
+			}
+			Result += TEXT("\n");
+		});
+	return Result;
+}
+
+TValueOrError<FString, FString> FToolsetRegistryToolAdapterManager::GetToolsetSchemaText(
+	const FString& ToolsetName) const
+{
+	UE::ToolsetRegistry::FToolsetRegistry* ToolsetRegistry = GetToolsetRegistry();
+	if (!ToolsetRegistry)
+	{
+		return MakeError(TEXT("ToolsetRegistry not available."));
 	}
-	RegisteredAdapters.Reset();
+
+	TSharedPtr<UE::ToolsetRegistry::FToolset> Toolset = ToolsetRegistry->Find(ToolsetName);
+	if (!Toolset.IsValid())
+	{
+		TArray<FString> Names;
+		ToolsetRegistry->ForEachToolset(
+			[&Names](const FString& Name, const UE::ToolsetRegistry::FToolset&)
+			{
+				Names.Add(Name);
+			});
+		return MakeError(FString::Printf(
+			TEXT("Toolset '%s' not found. Available toolsets: %s"),
+			*ToolsetName, *FString::Join(Names, TEXT(", "))));
+	}
+
+	return MakeValue(Toolset->GetToolsetJsonSchema());
 }
diff --git a/Engine/Plugins/Experimental/ModelContextProtocol/Source/ModelContextProtocolEditor/Private/ModelContextProtocolToolsetRegistryAdapter.h b/Engine/Plugins/Experimental/ModelContextProtocol/Source/ModelContextProtocolEditor/Private/ModelContextProtocolToolsetRegistryAdapter.h
--- a/Engine/Plugins/Experimental/ModelContextProtocol/Source/ModelContextProtocolEditor/Private/ModelContextProtocolToolsetRegistryAdapter.h
+++ b/Engine/Plugins/Experimental/ModelContextProtocol/Source/ModelContextProtocolEditor/Private/ModelContextProtocolToolsetRegistryAdapter.h
@@ -4,7 +4,7 @@
 
 #include "IModelContextProtocolTool.h"
 
-namespace UE::ToolsetRegistry { struct FToolDescriptor; class FToolsetRegistry; }
+namespace UE::ToolsetRegistry { struct FToolDescriptor; }
 
 /**
  * Adapts a ToolsetRegistry tool to the IModelContextProtocolTool interface.
@@ -27,15 +27,50 @@ private:
 
 /**
  * Registers and manages ToolsetRegistry tool adapters with the MCP module.
+ * Supports eager mode (all tools registered upfront) and deferred mode
+ * (list_toolsets, describe_toolset, and load_toolset registered for
+ * on-demand discovery and loading).
  */
 struct FToolsetRegistryToolAdapterManager
 {
-	/** Enumerates ToolsetRegistry tools and registers IModelContextProtocolTool adapters via AddTool(). */
+	~FToolsetRegistryToolAdapterManager()
+	{
+		DeregisterTools(false);
+	}
+
+	/**
+	 * Enumerates ToolsetRegistry tools and registers them with the MCP module.
+	 * In eager mode, registers all tool adapters directly. In deferred mode,
+	 * registers list_toolsets, describe_toolset, and load_toolset for
+	 * on-demand discovery and loading.
+	 */
 	void RegisterTools();
 
-	/** Removes all previously registered adapter tools. */
-	void DeregisterTools();
+	/** Removes all previously registered adapter tools. Optionally broadcasts tools/list_changed. */
+	void DeregisterTools(bool bBroadcast = true);
+
+	/**
+	 * Load a specific toolset on demand (deferred mode).
+	 * If the toolset is not yet loaded, registers its individual tool adapters
+	 * and broadcasts tools/list_changed. If already loaded, returns true immediately.
+	 * @param ToolsetName The toolset to load.
+	 * @param OutResultText Output message describing what was loaded or error details.
+	 * @return true if the toolset was loaded successfully or was already loaded.
+	 *         False if not found or contains no tools.
+	 */
+	bool RegisterToolset(const FString& ToolsetName, FString& OutResultText);
 
 private:
+	/** Register all tools from a toolset's JSON schema string. Returns the number of tools registered. */
+	int32 RegisterToolsFromSchema(const FString& ToolsetSchemaJson);
+
+	/** Returns a compact text listing of all toolsets (name + description). */
+	FString GetToolsetCatalogText() const;
+
+	/** Returns the full JSON schema for a named toolset, or an error. */
+	TValueOrError<FString, FString> GetToolsetSchemaText(const FString& ToolsetName) const;
+
 	TArray<TSharedRef<FToolsetRegistryToolAdapter>> RegisteredAdapters;
+	TArray<TSharedRef<IModelContextProtocolTool>> DeferredTools;
+	TSet<FString> LoadedToolsetNames;
 };
diff --git a/Engine/Plugins/Experimental/ModelContextProtocol/Source/ModelContextProtocolEditorTests/Private/Tests/ModelContextProtocolToolsetRegistryTests.cpp b/Engine/Plugins/Experimental/ModelContextProtocol/Source/ModelContextProtocolEditorTests/Private/Tests/ModelContextProtocolToolsetRegistryTests.cpp
--- a/Engine/Plugins/Experimental/ModelContextProtocol/Source/ModelContextProtocolEditorTests/Private/Tests/ModelContextProtocolToolsetRegistryTests.cpp
+++ b/Engine/Plugins/Experimental/ModelContextProtocol/Source/ModelContextProtocolEditorTests/Private/Tests/ModelContextProtocolToolsetRegistryTests.cpp
@@ -24,10 +24,14 @@
 static const FString MockToolsetName = TEXT("ModelContextProtocolEditorTests.MockToolsetDefinition");
 static const FString GreetToolName = MockToolsetName + TEXT(".Greet");
 static const FString AddToolName = MockToolsetName + TEXT(".Add");
+static const FString ListToolsetsToolName = TEXT("list_toolsets");
+static const FString DescribeToolsetToolName = TEXT("describe_toolset");
+static const FString LoadToolsetToolName = TEXT("load_toolset");
 
 BEGIN_DEFINE_SPEC(FModelContextProtocolToolsetRegistryTests, "AI.ModelContextProtocol.ToolsetRegistry", EAutomationTestFlags::ProductFilter | EAutomationTestFlags_ApplicationContextMask)
 	FString SessionId;
 	TSharedPtr<UE::ToolsetRegistry::FBlueprintLibraryToolset> MockToolset;
+	bool bSavedDeferredToolLoading = false;
 
 	void SendJsonRpcRequest(const TSharedRef<FJsonObject>& JsonRpcBody, const TFunction<void(FHttpResponsePtr, bool)>& OnComplete, const FString& InSessionId = FString(), TOptional<float> ActivityTimeout = {})
 	{
@@ -116,288 +120,572 @@ BEGIN_DEFINE_SPEC(FModelContextProtocolToolsetRegistryTests, "AI.ModelContextPro
 		});
 		DeleteRequest->ProcessRequest();
 	}
-END_DEFINE_SPEC(FModelContextProtocolToolsetRegistryTests)
-
-void FModelContextProtocolToolsetRegistryTests::Define()
-{
-	using namespace UE::ModelContextProtocol::Tests;
 
-	BeforeEach([this]()
+	void LoadMockToolsetViaRpc(const FDoneDelegate& Done)
 	{
-		SessionId.Reset();
+		using namespace UE::ModelContextProtocol::Tests;
 
-		// Register mock toolset with ToolsetRegistry
-		if (UToolsetRegistrySubsystem* ToolsetRegistrySubsystem = GEditor ? GEditor->GetEditorSubsystem<UToolsetRegistrySubsystem>() : nullptr)
-		{
-			MockToolset = MakeShared<UE::ToolsetRegistry::FBlueprintLibraryToolset>(UMockToolsetDefinition::StaticClass());
-			ToolsetRegistrySubsystem->ToolsetRegistry.RegisterToolset(MockToolset);
-		}
+		TSharedRef<FJsonObject> ToolArguments = MakeShared<FJsonObject>();
+		ToolArguments->SetStringField(TEXT("toolset_name"), MockToolsetName);
 
-		// Trigger adapter re-registration so the mock toolset's tools are available via IModelContextProtocolTool
-		IModelContextProtocolModule::GetChecked().RefreshTools();
-	});
+		TSharedRef<FJsonObject> CallParams = MakeShared<FJsonObject>();
+		CallParams->SetStringField(TEXT("name"), LoadToolsetToolName);
+		CallParams->SetObjectField(TEXT("arguments"), ToolArguments);
+
+		TSharedRef<FJsonObject> CallRequest = MakeJsonRpcRequest(TEXT("tools/call"), MakeShared<FJsonValueNumber>(100), CallParams);
+		// tools/call returns an SSE stream; use a short activity timeout since the tool completes synchronously.
+		SendJsonRpcRequest(CallRequest, [this, Done](FHttpResponsePtr Response, bool bConnected)
+		{
+			if (TestTrue("Should connect for load_toolset", bConnected) && TestTrue("Should have response", Response.IsValid()))
+			{
+				TSharedPtr<FJsonObject> JsonResponse = ParseSseJsonResponse(Response);
+				if (JsonResponse.IsValid())
+				{
+					const TSharedPtr<FJsonObject>* ResultObject = nullptr;
+					if (JsonResponse->TryGetObjectField(TEXT("result"), ResultObject))
+					{
+						bool bIsError = false;
+						(*ResultObject)->TryGetBoolField(TEXT("isError"), bIsError);
+						TestFalse("load_toolset should not return error", bIsError);
+					}
+				}
+			}
+			Done.Execute();
+		}, SessionId, 2.0f);
+	}
 
-	AfterEach([this]()
+	void InitializeAndPrepareTools(bool bDeferred, const FDoneDelegate& Done)
 	{
-		if (MockToolset.IsValid())
+		InitializeSession(FDoneDelegate::CreateLambda([this, bDeferred, Done]()
 		{
-			if (UToolsetRegistrySubsystem* ToolsetRegistrySubsystem = GEditor ? GEditor->GetEditorSubsystem<UToolsetRegistrySubsystem>() : nullptr)
+			if (SessionId.IsEmpty() || !bDeferred)
 			{
-				ToolsetRegistrySubsystem->ToolsetRegistry.UnregisterToolset(MockToolset);
+				Done.Execute();
+				return;
 			}
-			MockToolset.Reset();
-		}
-	});
+			LoadMockToolsetViaRpc(Done);
+		}));
+	}
+END_DEFINE_SPEC(FModelContextProtocolToolsetRegistryTests)
+
+void FModelContextProtocolToolsetRegistryTests::Define()
+{
+	using namespace UE::ModelContextProtocol::Tests;
 
-	Describe("ToolsetRegistry Tool Listing", [this]()
+	for (bool bDeferred : {false, true})
 	{
-		LatentIt("should include ToolsetRegistry tools in tools/list", [this](const FDoneDelegate& Done)
+		FString ModeName = bDeferred ? TEXT("Deferred") : TEXT("Eager");
+
+		Describe(*ModeName, [this, bDeferred]()
 		{
-			InitializeSession(FDoneDelegate::CreateLambda([this, Done]()
+			BeforeEach([this, bDeferred]()
 			{
-				if (SessionId.IsEmpty())
+				SessionId.Reset();
+
+				// Save and override deferred loading mode.
+				bSavedDeferredToolLoading = UE::ModelContextProtocol::bDeferredToolLoading;
+				UE::ModelContextProtocol::bDeferredToolLoading = bDeferred;
+
+				// Register mock toolset with ToolsetRegistry.
+				if (UToolsetRegistrySubsystem* ToolsetRegistrySubsystem = GEditor ? GEditor->GetEditorSubsystem<UToolsetRegistrySubsystem>() : nullptr)
 				{
-					Done.Execute();
-					return;
+					MockToolset = MakeShared<UE::ToolsetRegistry::FBlueprintLibraryToolset>(UMockToolsetDefinition::StaticClass());
+					ToolsetRegistrySubsystem->ToolsetRegistry.RegisterToolset(MockToolset);
 				}
 
-				TSharedRef<FJsonObject> ListRequest = MakeJsonRpcRequest(TEXT("tools/list"), MakeShared<FJsonValueNumber>(2));
-				SendJsonRpcRequest(ListRequest, [this, Done](FHttpResponsePtr Response, bool bConnected)
+				// Trigger adapter re-registration so the mock toolset's tools are available via IModelContextProtocolTool.
+				IModelContextProtocolModule::GetChecked().RefreshTools();
+			});
+
+			AfterEach([this]()
+			{
+				// Restore deferred loading mode.
+				UE::ModelContextProtocol::bDeferredToolLoading = bSavedDeferredToolLoading;
+
+				if (MockToolset.IsValid())
+				{
+					if (UToolsetRegistrySubsystem* ToolsetRegistrySubsystem = GEditor ? GEditor->GetEditorSubsystem<UToolsetRegistrySubsystem>() : nullptr)
+					{
+						ToolsetRegistrySubsystem->ToolsetRegistry.UnregisterToolset(MockToolset);
+					}
+					MockToolset.Reset();
+				}
+
+				// Re-sync MCP module tools with the restored CVar state.
+				IModelContextProtocolModule::GetChecked().RefreshTools();
+			});
+
+			Describe("Tool Listing", [this, bDeferred]()
+			{
+				LatentIt("should include ToolsetRegistry tools in tools/list", [this, bDeferred](const FDoneDelegate& Done)
 				{
-					if (TestTrue("Should connect", bConnected) && TestTrue("Should have response", Response.IsValid()))
+					InitializeAndPrepareTools(bDeferred, FDoneDelegate::CreateLambda([this, Done]()
 					{
-						TSharedPtr<FJsonObject> JsonResponse = ParseJsonResponse(Response);
-						if (TestTrue("Should parse response", JsonResponse.IsValid()))
+						if (SessionId.IsEmpty())
 						{
-							const TSharedPtr<FJsonObject>* ResultObject;
-							if (TestTrue("Should have result", JsonResponse->TryGetObjectField(TEXT("result"), ResultObject)))
+							Done.Execute();
+							return;
+						}
+
+						TSharedRef<FJsonObject> ListRequest = MakeJsonRpcRequest(TEXT("tools/list"), MakeShared<FJsonValueNumber>(2));
+						SendJsonRpcRequest(ListRequest, [this, Done](FHttpResponsePtr Response, bool bConnected)
+						{
+							if (TestTrue("Should connect", bConnected) && TestTrue("Should have response", Response.IsValid()))
 							{
-								const TArray<TSharedPtr<FJsonValue>>* ToolsArray;
-								if (TestTrue("Should have tools array", (*ResultObject)->TryGetArrayField(TEXT("tools"), ToolsArray)))
+								TSharedPtr<FJsonObject> JsonResponse = ParseJsonResponse(Response);
+								if (TestTrue("Should parse response", JsonResponse.IsValid()))
 								{
-									bool bFoundGreet = false;
-									bool bFoundAdd = false;
-									for (const TSharedPtr<FJsonValue>& ToolValue : *ToolsArray)
+									const TSharedPtr<FJsonObject>* ResultObject;
+									if (TestTrue("Should have result", JsonResponse->TryGetObjectField(TEXT("result"), ResultObject)))
 									{
-										const TSharedPtr<FJsonObject>* ToolObject;
-										if (ToolValue->TryGetObject(ToolObject))
+										const TArray<TSharedPtr<FJsonValue>>* ToolsArray;
+										if (TestTrue("Should have tools array", (*ResultObject)->TryGetArrayField(TEXT("tools"), ToolsArray)))
 										{
-											FString ToolName;
-											(*ToolObject)->TryGetStringField(TEXT("name"), ToolName);
-											if (ToolName == GreetToolName)
-											{
-												bFoundGreet = true;
-												FString ToolDescription;
-												(*ToolObject)->TryGetStringField(TEXT("description"), ToolDescription);
-												TestFalse("Greet tool should have a description", ToolDescription.IsEmpty());
-												const TSharedPtr<FJsonObject>* InputSchemaField;
-												TestTrue("Greet tool should have inputSchema", (*ToolObject)->TryGetObjectField(TEXT("inputSchema"), InputSchemaField));
-											}
-											else if (ToolName == AddToolName)
+											bool bFoundGreet = false;
+											bool bFoundAdd = false;
+											for (const TSharedPtr<FJsonValue>& ToolValue : *ToolsArray)
 											{
-												bFoundAdd = true;
-												const TSharedPtr<FJsonObject>* InputSchemaField;
-												TestTrue("Add tool should have inputSchema", (*ToolObject)->TryGetObjectField(TEXT("inputSchema"), InputSchemaField));
+												const TSharedPtr<FJsonObject>* ToolObject;
+												if (ToolValue->TryGetObject(ToolObject))
+												{
+													FString ToolName;
+													(*ToolObject)->TryGetStringField(TEXT("name"), ToolName);
+													if (ToolName == GreetToolName)
+													{
+														bFoundGreet = true;
+														FString ToolDescription;
+														(*ToolObject)->TryGetStringField(TEXT("description"), ToolDescription);
+														TestFalse("Greet tool should have a description", ToolDescription.IsEmpty());
+														const TSharedPtr<FJsonObject>* InputSchemaField;
+														TestTrue("Greet tool should have inputSchema", (*ToolObject)->TryGetObjectField(TEXT("inputSchema"), InputSchemaField));
+													}
+													else if (ToolName == AddToolName)
+													{
+														bFoundAdd = true;
+														const TSharedPtr<FJsonObject>* InputSchemaField;
+														TestTrue("Add tool should have inputSchema", (*ToolObject)->TryGetObjectField(TEXT("inputSchema"), InputSchemaField));
+													}
+												}
 											}
+											TestTrue("Should find Greet tool from ToolsetRegistry", bFoundGreet);
+											TestTrue("Should find Add tool from ToolsetRegistry", bFoundAdd);
 										}
 									}
-									TestTrue("Should find Greet tool from ToolsetRegistry", bFoundGreet);
-									TestTrue("Should find Add tool from ToolsetRegistry", bFoundAdd);
 								}
 							}
-						}
-					}
-					CleanupSession(Done);
-				}, SessionId);
-			}));
-		});
-	});
+							CleanupSession(Done);
+						}, SessionId);
+					}));
+				});
+			});
 
-	Describe("ToolsetRegistry Tool Execution", [this]()
-	{
-		LatentIt("should execute ToolsetRegistry tool via tools/call", [this](const FDoneDelegate& Done)
-		{
-			InitializeSession(FDoneDelegate::CreateLambda([this, Done]()
+			Describe("Tool Execution", [this, bDeferred]()
 			{
-				if (SessionId.IsEmpty())
+				LatentIt("should execute ToolsetRegistry tool via tools/call", [this, bDeferred](const FDoneDelegate& Done)
 				{
-					Done.Execute();
-					return;
-				}
+					InitializeAndPrepareTools(bDeferred, FDoneDelegate::CreateLambda([this, Done]()
+					{
+						if (SessionId.IsEmpty())
+						{
+							Done.Execute();
+							return;
+						}
 
-				TSharedRef<FJsonObject> ToolArguments = MakeShared<FJsonObject>();
-				ToolArguments->SetStringField(TEXT("Name"), TEXT("World"));
+						TSharedRef<FJsonObject> ToolArguments = MakeShared<FJsonObject>();
+						ToolArguments->SetStringField(TEXT("Name"), TEXT("World"));
 
-				TSharedRef<FJsonObject> CallParams = MakeShared<FJsonObject>();
-				CallParams->SetStringField(TEXT("name"), GreetToolName);
-				CallParams->SetObjectField(TEXT("arguments"), ToolArguments);
+						TSharedRef<FJsonObject> CallParams = MakeShared<FJsonObject>();
+						CallParams->SetStringField(TEXT("name"), GreetToolName);
+						CallParams->SetObjectField(TEXT("arguments"), ToolArguments);
 
-				TSharedRef<FJsonObject> CallRequest = MakeJsonRpcRequest(TEXT("tools/call"), MakeShared<FJsonValueNumber>(3), CallParams);
-				// tools/call returns an SSE stream; use a short activity timeout since the tool completes synchronously
-				SendJsonRpcRequest(CallRequest, [this, Done](FHttpResponsePtr Response, bool bConnected)
-				{
-					if (TestTrue("Should have response", Response.IsValid()))
-					{
-						TSharedPtr<FJsonObject> JsonResponse = ParseSseJsonResponse(Response);
-						if (JsonResponse.IsValid())
+						TSharedRef<FJsonObject> CallRequest = MakeJsonRpcRequest(TEXT("tools/call"), MakeShared<FJsonValueNumber>(3), CallParams);
+						// tools/call returns an SSE stream; use a short activity timeout since the tool completes synchronously.
+						SendJsonRpcRequest(CallRequest, [this, Done](FHttpResponsePtr Response, bool bConnected)
 						{
-							const TSharedPtr<FJsonObject>* ResultObject;
-							if (TestTrue("Should have result", JsonResponse->TryGetObjectField(TEXT("result"), ResultObject)))
+							if (TestTrue("Should have response", Response.IsValid()))
 							{
-								const TArray<TSharedPtr<FJsonValue>>* ContentArray;
-								if (TestTrue("Should have content", (*ResultObject)->TryGetArrayField(TEXT("content"), ContentArray))
-									&& TestTrue("Content should have at least 1 element", ContentArray->Num() >= 1))
+								TSharedPtr<FJsonObject> JsonResponse = ParseSseJsonResponse(Response);
+								if (JsonResponse.IsValid())
 								{
-									const TSharedPtr<FJsonObject>* ContentObject;
-									if ((*ContentArray)[0]->TryGetObject(ContentObject))
+									const TSharedPtr<FJsonObject>* ResultObject;
+									if (TestTrue("Should have result", JsonResponse->TryGetObjectField(TEXT("result"), ResultObject)))
 									{
-										FString Text;
-										if ((*ContentObject)->TryGetStringField(TEXT("text"), Text))
+										const TArray<TSharedPtr<FJsonValue>>* ContentArray;
+										if (TestTrue("Should have content", (*ResultObject)->TryGetArrayField(TEXT("content"), ContentArray))
+											&& TestTrue("Content should have at least 1 element", ContentArray->Num() >= 1))
 										{
-											TestTrue("Tool result should contain greeting", Text.Contains(TEXT("Hello, World!")));
+											const TSharedPtr<FJsonObject>* ContentObject;
+											if ((*ContentArray)[0]->TryGetObject(ContentObject))
+											{
+												FString Text;
+												if ((*ContentObject)->TryGetStringField(TEXT("text"), Text))
+												{
+													TestTrue("Tool result should contain greeting", Text.Contains(TEXT("Hello, World!")));
+												}
+											}
 										}
 									}
 								}
+								else
+								{
+									int32 ResponseCode = Response->GetResponseCode();
+									TestTrue("Should return success status for tool call", ResponseCode >= 200 && ResponseCode < 300);
+								}
 							}
-						}
-						else
+							CleanupSession(Done);
+						}, SessionId, 2.0f);
+					}));
+				});
+
+				LatentIt("should execute ToolsetRegistry tool with numeric result via tools/call", [this, bDeferred](const FDoneDelegate& Done)
+				{
+					InitializeAndPrepareTools(bDeferred, FDoneDelegate::CreateLambda([this, Done]()
+					{
+						if (SessionId.IsEmpty())
 						{
-							int32 ResponseCode = Response->GetResponseCode();
-							TestTrue("Should return success status for tool call", ResponseCode >= 200 && ResponseCode < 300);
+							Done.Execute();
+							return;
 						}
-					}
-					CleanupSession(Done);
-				}, SessionId, 2.0f);
-			}));
-		});
 
-		LatentIt("should execute ToolsetRegistry tool with numeric result via tools/call", [this](const FDoneDelegate& Done)
-		{
-			InitializeSession(FDoneDelegate::CreateLambda([this, Done]()
-			{
-				if (SessionId.IsEmpty())
-				{
-					Done.Execute();
-					return;
-				}
+						TSharedRef<FJsonObject> ToolArguments = MakeShared<FJsonObject>();
+						ToolArguments->SetNumberField(TEXT("A"), 3);
+						ToolArguments->SetNumberField(TEXT("B"), 7);
 
-				TSharedRef<FJsonObject> ToolArguments = MakeShared<FJsonObject>();
-				ToolArguments->SetNumberField(TEXT("A"), 3);
-				ToolArguments->SetNumberField(TEXT("B"), 7);
+						TSharedRef<FJsonObject> CallParams = MakeShared<FJsonObject>();
+						CallParams->SetStringField(TEXT("name"), AddToolName);
+						CallParams->SetObjectField(TEXT("arguments"), ToolArguments);
 
-				TSharedRef<FJsonObject> CallParams = MakeShared<FJsonObject>();
-				CallParams->SetStringField(TEXT("name"), AddToolName);
-				CallParams->SetObjectField(TEXT("arguments"), ToolArguments);
+						TSharedRef<FJsonObject> CallRequest = MakeJsonRpcRequest(TEXT("tools/call"), MakeShared<FJsonValueNumber>(4), CallParams);
+						// tools/call returns an SSE stream; use a short activity timeout since the tool completes synchronously.
+						SendJsonRpcRequest(CallRequest, [this, Done](FHttpResponsePtr Response, bool bConnected)
+						{
+							if (TestTrue("Should have response", Response.IsValid()))
+							{
+								TSharedPtr<FJsonObject> JsonResponse = ParseSseJsonResponse(Response);
+								if (JsonResponse.IsValid())
+								{
+									const TSharedPtr<FJsonObject>* ResultObject;
+									if (TestTrue("Should have result", JsonResponse->TryGetObjectField(TEXT("result"), ResultObject)))
+									{
+										const TArray<TSharedPtr<FJsonValue>>* ContentArray;
+										if (TestTrue("Should have content", (*ResultObject)->TryGetArrayField(TEXT("content"), ContentArray))
+											&& TestTrue("Content should have at least 1 element", ContentArray->Num() >= 1))
+										{
+											const TSharedPtr<FJsonObject>* ContentObject;
+											if ((*ContentArray)[0]->TryGetObject(ContentObject))
+											{
+												FString Text;
+												if ((*ContentObject)->TryGetStringField(TEXT("text"), Text))
+												{
+													TestTrue("Result should contain the sum", Text.Contains(TEXT("10")));
+												}
+											}
+										}
+									}
+								}
+								else
+								{
+									int32 ResponseCode = Response->GetResponseCode();
+									TestTrue("Should return success status for tool call", ResponseCode >= 200 && ResponseCode < 300);
+								}
+							}
+							CleanupSession(Done);
+						}, SessionId, 2.0f);
+					}));
+				});
+			});
 
-				TSharedRef<FJsonObject> CallRequest = MakeJsonRpcRequest(TEXT("tools/call"), MakeShared<FJsonValueNumber>(4), CallParams);
-				// tools/call returns an SSE stream; use a short activity timeout since the tool completes synchronously
-				SendJsonRpcRequest(CallRequest, [this, Done](FHttpResponsePtr Response, bool bConnected)
+			Describe("Tool Deduplication", [this, bDeferred]()
+			{
+				LatentIt("should deduplicate when MCP and ToolsetRegistry have same tool name", [this, bDeferred](const FDoneDelegate& Done)
 				{
-					if (TestTrue("Should have response", Response.IsValid()))
+					InitializeAndPrepareTools(bDeferred, FDoneDelegate::CreateLambda([this, Done]()
 					{
-						TSharedPtr<FJsonObject> JsonResponse = ParseSseJsonResponse(Response);
-						if (JsonResponse.IsValid())
+						if (SessionId.IsEmpty())
+						{
+							Done.Execute();
+							return;
+						}
+
+						// Register a deprecated MCP tool with the same qualified name as a ToolsetRegistry tool.
+						// AddTool() rejects duplicates at registration time, so the ToolsetRegistry adapter's
+						// tool (already registered via BeforeEach + InitializeAndPrepareTools) takes precedence.
+						TSharedPtr<FMockModelContextProtocolTool> MockMcpTool = MakeShared<FMockModelContextProtocolTool>();
+						MockMcpTool->Name = GreetToolName;
+						MockMcpTool->Description = TEXT("Deprecated MCP version of Greet tool");
+						MockMcpTool->InputSchema = FMockModelContextProtocolTool::MakeTestInputSchema(
+							{{TEXT("Name"), TEXT("string")}},
+							{TEXT("Name")});
+
+						IModelContextProtocolModule& Module = IModelContextProtocolModule::GetChecked();
+						AddExpectedError(TEXT("is already registered"), EAutomationExpectedErrorFlags::Contains);
+						TestFalse("Duplicate tool should be rejected by AddTool", Module.AddTool(MockMcpTool.ToSharedRef()));
+
+						TSharedRef<FJsonObject> ListRequest = MakeJsonRpcRequest(TEXT("tools/list"), MakeShared<FJsonValueNumber>(5));
+						SendJsonRpcRequest(ListRequest, [this, Done, MockMcpTool](FHttpResponsePtr Response, bool bConnected)
 						{
-							const TSharedPtr<FJsonObject>* ResultObject;
-							if (TestTrue("Should have result", JsonResponse->TryGetObjectField(TEXT("result"), ResultObject)))
+							if (TestTrue("Should connect", bConnected) && TestTrue("Should have response", Response.IsValid()))
 							{
-								const TArray<TSharedPtr<FJsonValue>>* ContentArray;
-								if (TestTrue("Should have content", (*ResultObject)->TryGetArrayField(TEXT("content"), ContentArray))
-									&& TestTrue("Content should have at least 1 element", ContentArray->Num() >= 1))
+								TSharedPtr<FJsonObject> JsonResponse = ParseJsonResponse(Response);
+								if (TestTrue("Should parse response", JsonResponse.IsValid()))
 								{
-									const TSharedPtr<FJsonObject>* ContentObject;
-									if ((*ContentArray)[0]->TryGetObject(ContentObject))
+									const TSharedPtr<FJsonObject>* ResultObject;
+									if (TestTrue("Should have result", JsonResponse->TryGetObjectField(TEXT("result"), ResultObject)))
 									{
-										FString Text;
-										if ((*ContentObject)->TryGetStringField(TEXT("text"), Text))
+										const TArray<TSharedPtr<FJsonValue>>* ToolsArray;
+										if (TestTrue("Should have tools array", (*ResultObject)->TryGetArrayField(TEXT("tools"), ToolsArray)))
 										{
-											TestTrue("Result should contain the sum", Text.Contains(TEXT("10")));
+											int32 GreetCount = 0;
+											FString GreetDescription;
+											for (const TSharedPtr<FJsonValue>& ToolValue : *ToolsArray)
+											{
+												const TSharedPtr<FJsonObject>* ToolObject;
+												if (ToolValue->TryGetObject(ToolObject))
+												{
+													FString ToolName;
+													(*ToolObject)->TryGetStringField(TEXT("name"), ToolName);
+													if (ToolName == GreetToolName)
+													{
+														GreetCount++;
+														(*ToolObject)->TryGetStringField(TEXT("description"), GreetDescription);
+													}
+												}
+											}
+											TestEqual("Should have exactly one Greet tool (deduplicated)", GreetCount, 1);
+											TestFalse("ToolsetRegistry tool should take precedence (not the deprecated MCP version)",
+												GreetDescription == TEXT("Deprecated MCP version of Greet tool"));
 										}
 									}
 								}
 							}
-						}
-						else
-						{
-							int32 ResponseCode = Response->GetResponseCode();
-							TestTrue("Should return success status for tool call", ResponseCode >= 200 && ResponseCode < 300);
-						}
-					}
-					CleanupSession(Done);
-				}, SessionId, 2.0f);
-			}));
-		});
-	});
 
-	Describe("ToolsetRegistry Tool Deduplication", [this]()
-	{
-		LatentIt("should deduplicate when MCP and ToolsetRegistry have same tool name", [this](const FDoneDelegate& Done)
-		{
-			InitializeSession(FDoneDelegate::CreateLambda([this, Done]()
+							CleanupSession(Done);
+						}, SessionId);
+					}));
+				});
+			});
+
+			// Deferred-only tests: verify behavior before toolsets are loaded.
+			if (bDeferred)
 			{
-				if (SessionId.IsEmpty())
+				Describe("Before Loading", [this]()
 				{
-					Done.Execute();
-					return;
-				}
+					LatentIt("should only show deferred tools before loading", [this](const FDoneDelegate& Done)
+					{
+						InitializeSession(FDoneDelegate::CreateLambda([this, Done]()
+						{
+							if (SessionId.IsEmpty())
+							{
+								Done.Execute();
+								return;
+							}
 
-				// Register a deprecated MCP tool with the same qualified name as a ToolsetRegistry tool.
-				// AddTool() rejects duplicates at registration time, so the ToolsetRegistry adapter's
-				// tool (already registered via RefreshTools in BeforeEach) takes precedence.
-				TSharedPtr<FMockModelContextProtocolTool> MockMcpTool = MakeShared<FMockModelContextProtocolTool>();
-				MockMcpTool->Name = GreetToolName;
-				MockMcpTool->Description = TEXT("Deprecated MCP version of Greet tool");
-				MockMcpTool->InputSchema = FMockModelContextProtocolTool::MakeTestInputSchema(
-					{{TEXT("Name"), TEXT("string")}},
-					{TEXT("Name")});
-
-				IModelContextProtocolModule& Module = IModelContextProtocolModule::GetChecked();
-				AddExpectedError(TEXT("is already registered"), EAutomationExpectedErrorFlags::Contains);
-				TestFalse("Duplicate tool should be rejected by AddTool", Module.AddTool(MockMcpTool.ToSharedRef()));
-
-				TSharedRef<FJsonObject> ListRequest = MakeJsonRpcRequest(TEXT("tools/list"), MakeShared<FJsonValueNumber>(5));
-				SendJsonRpcRequest(ListRequest, [this, Done, MockMcpTool](FHttpResponsePtr Response, bool bConnected)
-				{
-					if (TestTrue("Should connect", bConnected) && TestTrue("Should have response", Response.IsValid()))
+							// tools/list returns a regular JSON-RPC response (not SSE like tools/call).
+							TSharedRef<FJsonObject> ListRequest = MakeJsonRpcRequest(TEXT("tools/list"), MakeShared<FJsonValueNumber>(10));
+							SendJsonRpcRequest(ListRequest, [this, Done](FHttpResponsePtr Response, bool bConnected)
+							{
+								if (TestTrue("Should connect", bConnected) && TestTrue("Should have response", Response.IsValid()))
+								{
+									TSharedPtr<FJsonObject> JsonResponse = ParseJsonResponse(Response);
+									if (TestTrue("Should parse response", JsonResponse.IsValid()))
+									{
+										const TSharedPtr<FJsonObject>* ResultObject;
+										if (TestTrue("Should have result", JsonResponse->TryGetObjectField(TEXT("result"), ResultObject)))
+										{
+											const TArray<TSharedPtr<FJsonValue>>* ToolsArray;
+											if (TestTrue("Should have tools array", (*ResultObject)->TryGetArrayField(TEXT("tools"), ToolsArray)))
+											{
+												bool bFoundListToolsets = false;
+												bool bFoundDescribeToolset = false;
+												bool bFoundLoadToolset = false;
+												bool bFoundGreet = false;
+												bool bFoundAdd = false;
+												for (const TSharedPtr<FJsonValue>& ToolValue : *ToolsArray)
+												{
+													const TSharedPtr<FJsonObject>* ToolObject;
+													if (ToolValue->TryGetObject(ToolObject))
+													{
+														FString ToolName;
+														(*ToolObject)->TryGetStringField(TEXT("name"), ToolName);
+														if (ToolName == ListToolsetsToolName)
+														{
+															bFoundListToolsets = true;
+														}
+														else if (ToolName == DescribeToolsetToolName)
+														{
+															bFoundDescribeToolset = true;
+														}
+														else if (ToolName == LoadToolsetToolName)
+														{
+															bFoundLoadToolset = true;
+														}
+														else if (ToolName == GreetToolName)
+														{
+															bFoundGreet = true;
+														}
+														else if (ToolName == AddToolName)
+														{
+															bFoundAdd = true;
+														}
+													}
+												}
+												TestTrue("Should find list_toolsets tool", bFoundListToolsets);
+												TestTrue("Should find describe_toolset tool", bFoundDescribeToolset);
+												TestTrue("Should find load_toolset tool", bFoundLoadToolset);
+												TestFalse("Should NOT find Greet tool before loading", bFoundGreet);
+												TestFalse("Should NOT find Add tool before loading", bFoundAdd);
+											}
+										}
+									}
+								}
+								CleanupSession(Done);
+							}, SessionId);
+						}));
+					});
+
+					LatentIt("should list toolsets via list_toolsets", [this](const FDoneDelegate& Done)
 					{
-						TSharedPtr<FJsonObject> JsonResponse = ParseJsonResponse(Response);
-						if (TestTrue("Should parse response", JsonResponse.IsValid()))
+						InitializeSession(FDoneDelegate::CreateLambda([this, Done]()
 						{
-							const TSharedPtr<FJsonObject>* ResultObject;
-							if (TestTrue("Should have result", JsonResponse->TryGetObjectField(TEXT("result"), ResultObject)))
+							if (SessionId.IsEmpty())
 							{
-								const TArray<TSharedPtr<FJsonValue>>* ToolsArray;
-								if (TestTrue("Should have tools array", (*ResultObject)->TryGetArrayField(TEXT("tools"), ToolsArray)))
+								Done.Execute();
+								return;
+							}
+
+							TSharedRef<FJsonObject> CallParams = MakeShared<FJsonObject>();
+							CallParams->SetStringField(TEXT("name"), ListToolsetsToolName);
+							CallParams->SetObjectField(TEXT("arguments"), MakeShared<FJsonObject>());
+
+							TSharedRef<FJsonObject> CallRequest = MakeJsonRpcRequest(TEXT("tools/call"), MakeShared<FJsonValueNumber>(12), CallParams);
+							// tools/call returns an SSE stream; use a short activity timeout since the tool completes synchronously.
+							SendJsonRpcRequest(CallRequest, [this, Done](FHttpResponsePtr Response, bool bConnected)
+							{
+								if (TestTrue("Should have response", Response.IsValid()))
 								{
-									int32 GreetCount = 0;
-									FString GreetDescription;
-									for (const TSharedPtr<FJsonValue>& ToolValue : *ToolsArray)
+									TSharedPtr<FJsonObject> JsonResponse = ParseSseJsonResponse(Response);
+									if (JsonResponse.IsValid())
 									{
-										const TSharedPtr<FJsonObject>* ToolObject;
-										if (ToolValue->TryGetObject(ToolObject))
+										const TSharedPtr<FJsonObject>* ResultObject;
+										if (TestTrue("Should have result", JsonResponse->TryGetObjectField(TEXT("result"), ResultObject)))
 										{
-											FString ToolName;
-											(*ToolObject)->TryGetStringField(TEXT("name"), ToolName);
-											if (ToolName == GreetToolName)
+											const TArray<TSharedPtr<FJsonValue>>* ContentArray;
+											if (TestTrue("Should have content", (*ResultObject)->TryGetArrayField(TEXT("content"), ContentArray))
+												&& TestTrue("Content should have at least 1 element", ContentArray->Num() >= 1))
 											{
-												GreetCount++;
-												(*ToolObject)->TryGetStringField(TEXT("description"), GreetDescription);
+												const TSharedPtr<FJsonObject>* ContentObject;
+												if ((*ContentArray)[0]->TryGetObject(ContentObject))
+												{
+													FString Text;
+													if ((*ContentObject)->TryGetStringField(TEXT("text"), Text))
+													{
+														TestTrue("Should list mock toolset name", Text.Contains(TEXT("MockToolsetDefinition")));
+													}
+												}
 											}
 										}
 									}
-									TestEqual("Should have exactly one Greet tool (deduplicated)", GreetCount, 1);
-									TestFalse("ToolsetRegistry tool should take precedence (not the deprecated MCP version)",
-										GreetDescription == TEXT("Deprecated MCP version of Greet tool"));
 								}
+								CleanupSession(Done);
+							}, SessionId, 2.0f);
+						}));
+					});
+
+					LatentIt("should describe toolset via describe_toolset", [this](const FDoneDelegate& Done)
+					{
+						InitializeSession(FDoneDelegate::CreateLambda([this, Done]()
+						{
+							if (SessionId.IsEmpty())
+							{
+								Done.Execute();
+								return;
 							}
-						}
-					}
 
-					CleanupSession(Done);
-				}, SessionId);
-			}));
+							TSharedRef<FJsonObject> ToolArguments = MakeShared<FJsonObject>();
+							ToolArguments->SetStringField(TEXT("toolset_name"), MockToolsetName);
+
+							TSharedRef<FJsonObject> CallParams = MakeShared<FJsonObject>();
+							CallParams->SetStringField(TEXT("name"), DescribeToolsetToolName);
+							CallParams->SetObjectField(TEXT("arguments"), ToolArguments);
+
+							TSharedRef<FJsonObject> CallRequest = MakeJsonRpcRequest(TEXT("tools/call"), MakeShared<FJsonValueNumber>(13), CallParams);
+							// tools/call returns an SSE stream; use a short activity timeout since the tool completes synchronously.
+							SendJsonRpcRequest(CallRequest, [this, Done](FHttpResponsePtr Response, bool bConnected)
+							{
+								if (TestTrue("Should have response", Response.IsValid()))
+								{
+									TSharedPtr<FJsonObject> JsonResponse = ParseSseJsonResponse(Response);
+									if (JsonResponse.IsValid())
+									{
+										const TSharedPtr<FJsonObject>* ResultObject;
+										if (TestTrue("Should have result", JsonResponse->TryGetObjectField(TEXT("result"), ResultObject)))
+										{
+											const TArray<TSharedPtr<FJsonValue>>* ContentArray;
+											if (TestTrue("Should have content", (*ResultObject)->TryGetArrayField(TEXT("content"), ContentArray))
+												&& TestTrue("Content should have at least 1 element", ContentArray->Num() >= 1))
+											{
+												const TSharedPtr<FJsonObject>* ContentObject;
+												if ((*ContentArray)[0]->TryGetObject(ContentObject))
+												{
+													FString Text;
+													if ((*ContentObject)->TryGetStringField(TEXT("text"), Text))
+													{
+														TestTrue("Should contain Greet tool info", Text.Contains(TEXT("Greet")));
+														TestTrue("Should contain Add tool info", Text.Contains(TEXT("Add")));
+													}
+												}
+											}
+										}
+									}
+								}
+								CleanupSession(Done);
+							}, SessionId, 2.0f);
+						}));
+					});
+
+					LatentIt("should return error for unknown toolset", [this](const FDoneDelegate& Done)
+					{
+						InitializeSession(FDoneDelegate::CreateLambda([this, Done]()
+						{
+							if (SessionId.IsEmpty())
+							{
+								Done.Execute();
+								return;
+							}
+
+							TSharedRef<FJsonObject> ToolArguments = MakeShared<FJsonObject>();
+							ToolArguments->SetStringField(TEXT("toolset_name"), TEXT("NonexistentToolset"));
+
+							TSharedRef<FJsonObject> CallParams = MakeShared<FJsonObject>();
+							CallParams->SetStringField(TEXT("name"), LoadToolsetToolName);
+							CallParams->SetObjectField(TEXT("arguments"), ToolArguments);
+
+							TSharedRef<FJsonObject> CallRequest = MakeJsonRpcRequest(TEXT("tools/call"), MakeShared<FJsonValueNumber>(14), CallParams);
+							// tools/call returns an SSE stream; use a short activity timeout since the tool completes synchronously.
+							SendJsonRpcRequest(CallRequest, [this, Done](FHttpResponsePtr Response, bool bConnected)
+							{
+								if (TestTrue("Should have response", Response.IsValid()))
+								{
+									TSharedPtr<FJsonObject> JsonResponse = ParseSseJsonResponse(Response);
+									if (JsonResponse.IsValid())
+									{
+										const TSharedPtr<FJsonObject>* ResultObject;
+										if (TestTrue("Should have result", JsonResponse->TryGetObjectField(TEXT("result"), ResultObject)))
+										{
+											bool bIsError = false;
+											(*ResultObject)->TryGetBoolField(TEXT("isError"), bIsError);
+											TestTrue("Should return error for unknown toolset", bIsError);
+										}
+									}
+								}
+								CleanupSession(Done);
+							}, SessionId, 2.0f);
+						}));
+					});
+				});
+			}
 		});
-	});
+	}
 }
 
 #endif // WITH_DEV_AUTOMATION_TESTS && WITH_EDITOR
diff --git a/Engine/Plugins/Experimental/ToolsetRegistry/Source/ToolsetRegistry/Private/Tests/FakeToolset.cpp b/Engine/Plugins/Experimental/ToolsetRegistry/Source/ToolsetRegistry/Private/Tests/FakeToolset.cpp
--- a/Engine/Plugins/Experimental/ToolsetRegistry/Source/ToolsetRegistry/Private/Tests/FakeToolset.cpp
+++ b/Engine/Plugins/Experimental/ToolsetRegistry/Source/ToolsetRegistry/Private/Tests/FakeToolset.cpp
@@ -61,6 +61,11 @@ namespace UE::ToolsetRegistry
 		return TEXT("1.0");
 	}
 
+	FString FFakeToolset::GetToolsetDescription() const
+	{
+		return TEXT("Fake Toolset for testing.");
+	}
+
 	TSharedPtr<TPromise<TValueOrError<FString, FString>>> FFakeToolset::AddFakeToolCall(
 		const FString& ToolName,
 		TPromise<TValueOrError<FString, FString>>&& Promise)
diff --git a/Engine/Plugins/Experimental/ToolsetRegistry/Source/ToolsetRegistry/Private/Tests/FakeToolset.h b/Engine/Plugins/Experimental/ToolsetRegistry/Source/ToolsetRegistry/Private/Tests/FakeToolset.h
--- a/Engine/Plugins/Experimental/ToolsetRegistry/Source/ToolsetRegistry/Private/Tests/FakeToolset.h
+++ b/Engine/Plugins/Experimental/ToolsetRegistry/Source/ToolsetRegistry/Private/Tests/FakeToolset.h
@@ -34,6 +34,7 @@ public:
 	UE_API FString GetToolsetName() const override;
 
 	UE_API FString GetToolsetVersion() const override;
+	UE_API FString GetToolsetDescription() const override;
 
 	UE_API TSharedPtr<TPromise<TValueOrError<FString, FString>>> AddFakeToolCall(
 		const FString& ToolName,
diff --git a/Engine/Plugins/Experimental/ToolsetRegistry/Source/ToolsetRegistry/Private/ToolsetRegistry/ToolsetRegistry.cpp b/Engine/Plugins/Experimental/ToolsetRegistry/Source/ToolsetRegistry/Private/ToolsetRegistry/ToolsetRegistry.cpp
--- a/Engine/Plugins/Experimental/ToolsetRegistry/Source/ToolsetRegistry/Private/ToolsetRegistry/ToolsetRegistry.cpp
+++ b/Engine/Plugins/Experimental/ToolsetRegistry/Source/ToolsetRegistry/Private/ToolsetRegistry/ToolsetRegistry.cpp
@@ -150,6 +150,18 @@ namespace UE::ToolsetRegistry
 		return ToolsetHandler->ExecuteTool(ToolDescriptor.ToolName, JsonInput);
 	}
 
+	void FToolsetRegistry::ForEachToolset(
+		TFunctionRef<void(const FString& Name, const FToolset& Toolset)> Callback) const
+	{
+		for (const TPair<FString, TSharedPtr<FToolset>>& Pair : ToolsetHandlers)
+		{
+			if (Pair.Value.IsValid())
+			{
+				Callback(Pair.Key, *Pair.Value);
+			}
+		}
+	}
+
 	FString FToolsetRegistry::GetToolsetJsonSchemas() const
 	{
 		FString Result = TEXT("[");
diff --git a/Engine/Plugins/Experimental/ToolsetRegistry/Source/ToolsetRegistry/Public/ToolsetRegistry/BlueprintLibraryToolset.h b/Engine/Plugins/Experimental/ToolsetRegistry/Source/ToolsetRegistry/Public/ToolsetRegistry/BlueprintLibraryToolset.h
--- a/Engine/Plugins/Experimental/ToolsetRegistry/Source/ToolsetRegistry/Public/ToolsetRegistry/BlueprintLibraryToolset.h
+++ b/Engine/Plugins/Experimental/ToolsetRegistry/Source/ToolsetRegistry/Public/ToolsetRegistry/BlueprintLibraryToolset.h
@@ -52,9 +52,14 @@ namespace UE::ToolsetRegistry
 		virtual UClass* GetToolsetClass() const override { return ToolsetClass.Get(); }
 
 		// Returns the description of the toolset from the class's metadata.
-		FString GetToolsetDescription() const
+		virtual FString GetToolsetDescription() const override
 		{
-			return ToolsetClass.Pin() ? ToolsetClass->GetToolTipText().ToString() : FString();
+			const TStrongObjectPtr<UClass> Pinned = ToolsetClass.Pin();
+			if (Pinned)
+			{
+				return Pinned->GetToolTipText().ToString();
+			}
+			return FString();
 		}
 
 		UE_API TArray<FString> GetToolNames();
diff --git a/Engine/Plugins/Experimental/ToolsetRegistry/Source/ToolsetRegistry/Public/ToolsetRegistry/Toolset.h b/Engine/Plugins/Experimental/ToolsetRegistry/Source/ToolsetRegistry/Public/ToolsetRegistry/Toolset.h
--- a/Engine/Plugins/Experimental/ToolsetRegistry/Source/ToolsetRegistry/Public/ToolsetRegistry/Toolset.h
+++ b/Engine/Plugins/Experimental/ToolsetRegistry/Source/ToolsetRegistry/Public/ToolsetRegistry/Toolset.h
@@ -35,6 +35,9 @@ namespace UE::ToolsetRegistry
 		/// Returns the version of the toolset
 		virtual FString GetToolsetVersion() const = 0;
 
+		/// Returns a description of what this toolset does and when it should be used.
+		virtual FString GetToolsetDescription() const = 0;
+
 		/// Returns the UClass that implements this toolset, if any.
 		virtual UClass* GetToolsetClass() const { return nullptr; }
 	};
diff --git a/Engine/Plugins/Experimental/ToolsetRegistry/Source/ToolsetRegistry/Public/ToolsetRegistry/ToolsetRegistry.h b/Engine/Plugins/Experimental/ToolsetRegistry/Source/ToolsetRegistry/Public/ToolsetRegistry/ToolsetRegistry.h
--- a/Engine/Plugins/Experimental/ToolsetRegistry/Source/ToolsetRegistry/Public/ToolsetRegistry/ToolsetRegistry.h
+++ b/Engine/Plugins/Experimental/ToolsetRegistry/Source/ToolsetRegistry/Public/ToolsetRegistry/ToolsetRegistry.h
@@ -57,6 +57,10 @@ namespace UE::ToolsetRegistry
 		UE_API TFuture<TValueOrError<FString, FString>> ExecuteTool(
 			const FToolDescriptor& ToolDescriptor, const FString& JsonInput);
 
+		// Iterate all registered toolsets.
+		UE_API void ForEachToolset(
+			TFunctionRef<void(const FString& Name, const FToolset& Toolset)> Callback) const;
+
 		// Get the JSON schemas for all registered toolsets.
 		// Returns a JSON list concatenating all registered Toolset schemas.
 		UE_API FString GetToolsetJsonSchemas() const;
diff --git a/Engine/Plugins/Experimental/Toolsets/MCPClientToolset/Source/MCPClientToolset/Private/MCPClientToolset/MCPClientToolset.cpp b/Engine/Plugins/Experimental/Toolsets/MCPClientToolset/Source/MCPClientToolset/Private/MCPClientToolset/MCPClientToolset.cpp
--- a/Engine/Plugins/Experimental/Toolsets/MCPClientToolset/Source/MCPClientToolset/Private/MCPClientToolset/MCPClientToolset.cpp
+++ b/Engine/Plugins/Experimental/Toolsets/MCPClientToolset/Source/MCPClientToolset/Private/MCPClientToolset/MCPClientToolset.cpp
@@ -197,6 +197,11 @@ FString FMCPClientToolset::GetToolsetVersion() const
 	return ServerVersion.IsEmpty() ? TEXT("1.0") : ServerVersion;
 }
 
+FString FMCPClientToolset::GetToolsetDescription() const
+{
+	return Config.Description;
+}
+
 FString FMCPClientToolset::GetToolsetJsonSchema() const
 {
 	return CachedSchemaJson;
diff --git a/Engine/Plugins/Experimental/Toolsets/MCPClientToolset/Source/MCPClientToolset/Public/MCPClientToolset/MCPClientToolset.h b/Engine/Plugins/Experimental/Toolsets/MCPClientToolset/Source/MCPClientToolset/Public/MCPClientToolset/MCPClientToolset.h
--- a/Engine/Plugins/Experimental/Toolsets/MCPClientToolset/Source/MCPClientToolset/Public/MCPClientToolset/MCPClientToolset.h
+++ b/Engine/Plugins/Experimental/Toolsets/MCPClientToolset/Source/MCPClientToolset/Public/MCPClientToolset/MCPClientToolset.h
@@ -86,6 +86,7 @@ public:
 	UE_API virtual FString GetToolsetJsonSchema() const override;
 	UE_API virtual FString GetToolsetName()    const override;
 	UE_API virtual FString GetToolsetVersion() const override;
+	UE_API virtual FString GetToolsetDescription() const override;
 
 private:
 	// ---- SSE connection ----
